Editor: Deleting bogus IntroTutorials binaries
Integrate CL# 2125184 from 4.3 to Main

diff --git a/Engine/Build/JunkManifest.txt b/Engine/Build/JunkManifest.txt
--- a/Engine/Build/JunkManifest.txt
+++ b/Engine/Build/JunkManifest.txt
@@ -40,3 +40,5 @@ Engine\Binaries\Win64\*-TextureXboxOneFormat.*
 Shadow\Binaries\Win64\*ShadowLoadingScreen.*
 Engine\Binaries\Win64\UnrealRemoteAgent.*
 Engine\Binaries\DotNET\AutomationScripts\NoRedist\FotniteGame.Automation.*
+Engine\Plugins\Editor\IntroTutorials\Binaries\Win64\*-IntroTutorials*.*
+Engine\Plugins\Editor\IntroTutorials\Binaries\Mac\*-IntroTutorials*.*
