remove old texture format dll and add it to the junk deleter
#ue4

diff --git a/Engine/Build/JunkManifest.txt b/Engine/Build/JunkManifest.txt
--- a/Engine/Build/JunkManifest.txt
+++ b/Engine/Build/JunkManifest.txt
@@ -37,6 +37,7 @@ Engine\Binaries\Mac\UE4Editor.app\Contents\Frameworks\*-SlateStyle.*
 Engine\Binaries\Mac\*-MessagingDebugger.*
 Engine\Binaries\Win64\*-MessagingDebugger.*
 Engine\Binaries\Win64\*-TextureXboxOneFormat.*
+Engine\Binaries\Win64\*-TextureFormatATITC.*
 Shadow\Binaries\Win64\*ShadowLoadingScreen.*
 Engine\Binaries\Win64\UnrealRemoteAgent.*
 Engine\Binaries\DotNET\AutomationScripts\NoRedist\FotniteGame.Automation.*
