remove Android, HTML5, and IOS from the junk deleter
#ue4

diff --git a/Engine/Build/JunkManifest.txt b/Engine/Build/JunkManifest.txt
--- a/Engine/Build/JunkManifest.txt
+++ b/Engine/Build/JunkManifest.txt
@@ -13,10 +13,7 @@ Platform=XboxOne:Engine\Binaries\Win64\RocketEditor-*XboxOne*.*
 Platform=XboxOne:Engine\Binaries\Win32\RocketEditor-*XboxOne*.*
 Platform=XboxOne:Engine\Binaries\Win64\ShaderCompileWorker-*XboxOne*.*
 Platform=XboxOne:Engine\Binaries\Win32\ShaderCompileWorker-*XboxOne*.*
-Platform=IOS:Engine\Binaries\Win64\IOS\*.*
 Platform=Linux:Engine\Binaries\Win64\Linux\*.*
-Platform=Android:Engine\Binaries\Win64\Android\*.*
-Platform=HTML5:Engine\Binaries\Win64\HTML5\*.*
 Engine\Binaries\Win64\UE4Editor-LinuxTarget*.*
 Engine\Binaries\Win64\UE4Editor-AndroidTarget*.*
 Engine\Binaries\Win64\UE4Editor-IOSPlatformEditor*.*
