Add UHT script plugins to the junk manifest.

diff --git a/Engine/Build/JunkManifest.txt b/Engine/Build/JunkManifest.txt
--- a/Engine/Build/JunkManifest.txt
+++ b/Engine/Build/JunkManifest.txt
@@ -40,3 +40,10 @@ Engine\Binaries\Win64\UnrealRemoteAgent.*
 Engine\Binaries\DotNET\AutomationScripts\NoRedist\FotniteGame.Automation.*
 Engine\Plugins\Editor\IntroTutorials\Binaries\Win64\*-IntroTutorials*.*
 Engine\Plugins\Editor\IntroTutorials\Binaries\Mac\*-IntroTutorials*.*
+Engine\Plugins\Script\ScriptEditorPlugin\Binaries\Mac\*
+Engine\Plugins\Script\ScriptEditorPlugin\Binaries\Win64\*
+Engine\Plugins\Script\ScriptGeneratorPlugin\Binaries\Mac\*
+Engine\Plugins\Script\ScriptGeneratorPlugin\Binaries\Win64\*
+Engine\Plugins\Script\ScriptPlugin\Binaries\Mac\*
+Engine\Plugins\Script\ScriptPlugin\Binaries\Win32\*
+Engine\Plugins\Script\ScriptPlugin\Binaries\Win64\*
