Put the Android ES31 dll back in.

diff --git a/Engine/Build/JunkManifest.txt b/Engine/Build/JunkManifest.txt
--- a/Engine/Build/JunkManifest.txt
+++ b/Engine/Build/JunkManifest.txt
@@ -28,7 +28,6 @@ Engine\Binaries\Win64\*-SlateStyle*.*
 Engine\Binaries\Win64\*-CameraEditor*.*
 Engine\Binaries\Win64\*-CameraDebugger*.*
 Engine\Binaries\Win64\*-SessionLauncher*.*
-Engine\Binaries\Win64\Android\*-Android_ES31TargetPlatform.*
 Engine\Binaries\Mac\UnrealFrontend.app\Contents\Frameworks\*-CoreSlateStyle.*
 Engine\Binaries\Mac\UnrealFrontend.app\Contents\Frameworks\*-SlateStyle.*
 Engine\Binaries\Mac\UE4Editor.app\Contents\Frameworks\*-CoreSlateStyle.*
