- Fixed JunkManifest.txt to not delete AndroidPlatformEditor.dylib
- Fixed a null check for ini sections in MakeAPK

diff --git a/Engine/Build/JunkManifest.txt b/Engine/Build/JunkManifest.txt
--- a/Engine/Build/JunkManifest.txt
+++ b/Engine/Build/JunkManifest.txt
@@ -32,7 +32,6 @@ Engine\Binaries\Mac\UnrealFrontend.app\Contents\Frameworks\*-CoreSlateStyle.*
 Engine\Binaries\Mac\UnrealFrontend.app\Contents\Frameworks\*-SlateStyle.*
 Engine\Binaries\Mac\UE4Editor.app\Contents\Frameworks\*-CoreSlateStyle.*
 Engine\Binaries\Mac\UE4Editor.app\Contents\Frameworks\*-SlateStyle.*
-Engine\Binaries\Mac\UE4Editor.app\Contents\MacOS\*-AndroidPlatformEditor.*
 Engine\Binaries\Mac\UE4Editor.app\Contents\MacOS\*-AndroidRuntimeSettings.*
 Engine\Binaries\Mac\UE4Editor.app\Contents\MacOS\*-CameraDebugger.*
 Engine\Binaries\Mac\UE4Editor.app\Contents\MacOS\*-CameraEditor.*
