Back out revision 77 from //depot/UE4/Engine/Build/JunkManifest.txt
This cause causing CIS to fail because it was deleting these binaries that are still in use in main.
This change must be re-introduced after CL#2444913 is merged from the 4.7 branch, which is the change that moves the binary location.

diff --git a/Engine/Build/JunkManifest.txt b/Engine/Build/JunkManifest.txt
--- a/Engine/Build/JunkManifest.txt
+++ b/Engine/Build/JunkManifest.txt
@@ -51,6 +51,5 @@ Engine\Plugins\Script\ScriptGeneratorPlugin\Binaries\Win64\*
 Engine\Plugins\Script\ScriptPlugin\Binaries\Mac\*
 Engine\Plugins\Script\ScriptPlugin\Binaries\Win32\*
 Engine\Plugins\Script\ScriptPlugin\Binaries\Win64\*
-Engine\Plugins\Runtime\WebBrowserAsPlugin\Binaries
 Engine\Plugins\Experimental\Paper2D\Binaries
 Engine\Plugins\Experimental\Paper2D\Intermediate
