Remove references to BuildingRocket() from Target.cs.template. Think it will fix CIS failing to create temporary target file for TappyChicken.

diff --git a/Engine/Build/Target.cs.template b/Engine/Build/Target.cs.template
--- a/Engine/Build/Target.cs.template
+++ b/Engine/Build/Target.cs.template
@@ -26,46 +26,6 @@ public class {GAME_NAME}Target : TargetRules
 		{
 			OutExtraModuleNames.Add("OnlineSubsystemNull");
 		}
-
-
-		if (UnrealBuildTool.UnrealBuildTool.BuildingRocket())
-		{
-			OutExtraModuleNames.Add("GameMenuBuilder");
-
-			if (Target.Platform == UnrealTargetPlatform.Win32 || Target.Platform == UnrealTargetPlatform.Win64)
-			{
-				OutExtraModuleNames.Add("OnlineSubsystemNull");
-				OutExtraModuleNames.Add("OnlineSubsystemAmazon");
-				if (UEBuildConfiguration.bCompileSteamOSS == true)
-				{
-					OutExtraModuleNames.Add("OnlineSubsystemSteam");
-				}
-				OutExtraModuleNames.Add("OnlineSubsystemFacebook");
-			}
-			else if (Target.Platform == UnrealTargetPlatform.Mac || Target.Platform == UnrealTargetPlatform.Linux)
-			{
-				OutExtraModuleNames.Add("OnlineSubsystemNull");
-				if (UEBuildConfiguration.bCompileSteamOSS == true)
-				{
-					OutExtraModuleNames.Add("OnlineSubsystemSteam");
-				}
-			}
-			else if (Target.Platform == UnrealTargetPlatform.IOS)
-			{
-				OutExtraModuleNames.Add("OnlineSubsystemFacebook");
-				OutExtraModuleNames.Add("OnlineSubsystemIOS");
-				OutExtraModuleNames.Add("IOSAdvertising");
-			}
-			else if (Target.Platform == UnrealTargetPlatform.Android)
-			{
-				// @todo android: Add Android online subsystem
-				OutExtraModuleNames.Add("AndroidAdvertising");
-			}
-			else if (Target.Platform == UnrealTargetPlatform.HTML5)
-			{
-				OutExtraModuleNames.Add("OnlineSubsystemNull");
-			}
-		}
 	}
 
 	public override void SetupGlobalEnvironment(
@@ -74,24 +34,6 @@ public class {GAME_NAME}Target : TargetRules
 		ref CPPEnvironmentConfiguration OutCPPEnvironmentConfiguration
 		)
 	{
-		if (UnrealBuildTool.UnrealBuildTool.BuildingRocket())
-		{
-			UEBuildConfiguration.bCompileLeanAndMeanUE = true;
-
-			// Don't need editor or editor only data
-			UEBuildConfiguration.bBuildEditor = false;
-			UEBuildConfiguration.bBuildWithEditorOnlyData = false;
-
-			UEBuildConfiguration.bCompileAgainstEngine = true;
-
-			// no exports, so no need to verify that a .lib and .exp file was emitted by the linker.
-			OutLinkEnvironmentConfiguration.bHasExports = false;
-		}
-		else
-		{
-			// Tag it as a UE4Game build
-			OutCPPEnvironmentConfiguration.Definitions.Add("UE4GAME=1");
-		}
 	}
 	public override GUBPProjectOptions GUBP_IncludeProjectInPromotedBuild_EditorTypeOnly(UnrealTargetPlatform HostPlatform)
     {
