- Removing old SteamVR plugin from plugins experimental
- Updated junkmanifest to remove old plugins and symbols

diff --git a/Engine/Build/JunkManifest.txt b/Engine/Build/JunkManifest.txt
--- a/Engine/Build/JunkManifest.txt
+++ b/Engine/Build/JunkManifest.txt
@@ -52,6 +52,7 @@ Engine\Plugins\Script\ScriptGeneratorPlugin\Binaries\Win64\*
 Engine\Plugins\Script\ScriptPlugin\Binaries\Mac\*
 Engine\Plugins\Script\ScriptPlugin\Binaries\Win32\*
 Engine\Plugins\Script\ScriptPlugin\Binaries\Win64\*
+Engine\Plugins\Experimental\SteamVR\Binaries\Win64\*
 Engine\Plugins\Experimental\Paper2D\Binaries
 Engine\Plugins\Experimental\Paper2D\Intermediate
 Engine\Binaries\Mac\UE4Editor.app\Contents\MacOS\UE4Editor-Android_ES31TargetPlatform.dylib
diff --git a/Engine/Plugins/Experimental/SteamVR/Source/SteamVR/Private/SteamVRHMD.cpp b/Engine/Plugins/Experimental/SteamVR/Source/SteamVR/Private/SteamVRHMD.cpp
deleted file mode 100644
diff --git a/Engine/Plugins/Experimental/SteamVR/Source/SteamVR/Private/SteamVRHMD.h b/Engine/Plugins/Experimental/SteamVR/Source/SteamVR/Private/SteamVRHMD.h
deleted file mode 100644
diff --git a/Engine/Plugins/Experimental/SteamVR/Source/SteamVR/Private/SteamVRPrivatePCH.h b/Engine/Plugins/Experimental/SteamVR/Source/SteamVR/Private/SteamVRPrivatePCH.h
deleted file mode 100644
diff --git a/Engine/Plugins/Experimental/SteamVR/Source/SteamVR/Private/SteamVRRender.cpp b/Engine/Plugins/Experimental/SteamVR/Source/SteamVR/Private/SteamVRRender.cpp
deleted file mode 100644
diff --git a/Engine/Plugins/Experimental/SteamVR/Source/SteamVR/Public/ISteamVRPlugin.h b/Engine/Plugins/Experimental/SteamVR/Source/SteamVR/Public/ISteamVRPlugin.h
deleted file mode 100644
diff --git a/Engine/Plugins/Experimental/SteamVR/Source/SteamVR/SteamVR.Build.cs b/Engine/Plugins/Experimental/SteamVR/Source/SteamVR/SteamVR.Build.cs
deleted file mode 100644
diff --git a/Engine/Plugins/Experimental/SteamVR/SteamVR.uplugin b/Engine/Plugins/Experimental/SteamVR/SteamVR.uplugin
deleted file mode 100644
