Remove TwitchLiveStreaming plugin from the Rocket build.

diff --git a/Engine/Build/InstalledEngineFilters.ini b/Engine/Build/InstalledEngineFilters.ini
--- a/Engine/Build/InstalledEngineFilters.ini
+++ b/Engine/Build/InstalledEngineFilters.ini
@@ -24,6 +24,14 @@
 ; Plugins
 /Engine/Plugins/....uplugin
 /Engine/Plugins/.../Content/...
+/Engine/Plugins/.../Source/...
+-/Engine/Plugins/Runtime/TwitchLiveStreaming/...
+
+; Source code
+/Engine/Source/Runtime/...
+/Engine/Source/Developer/...
+/Engine/Source/Editor/...
+/Engine/Source/ThirdParty/Licenses/...
 
 ; Shaders
 /Engine/Shaders/...
@@ -93,13 +101,6 @@
 ; API docs
 /Engine/Documentation/CHM/API.chm
 
-; Source code
-/Engine/Source/Runtime/...
-/Engine/Source/Developer/...
-/Engine/Source/Editor/...
-/Engine/Source/ThirdParty/Licenses/...
-/Engine/Plugins/.../Source/...
-
 ; Starter content
 /Samples/StarterContent/Content/...
 /Samples/MobileStarterContent/Content/...
