Fixed build distribution issue with IOSApsalar plugin

diff --git a/Engine/Build/JunkManifest.txt b/Engine/Build/JunkManifest.txt
--- a/Engine/Build/JunkManifest.txt
+++ b/Engine/Build/JunkManifest.txt
@@ -58,6 +58,7 @@ Engine\Plugins\Script\ScriptPlugin\Binaries\Win64\*
 Engine\Plugins\Experimental\SteamVR\Binaries\Win64\*
 Engine\Plugins\Experimental\Paper2D\Binaries
 Engine\Plugins\Experimental\Paper2D\Intermediate
+Engine\Plugins\Runtime\Analytics\Apsalar\IOSApsalar\Binaries\Mac\*-IOSApsalar.*
 Engine\Plugins\Runtime\TwitchLiveStreaming\Binaries\Win32\TwitchLiveStreaming*.*
 Engine\Plugins\Runtime\TwitchLiveStreaming\Binaries\Win32\*-TwitchLiveStreaming*.*
 Engine\Plugins\Runtime\TwitchLiveStreaming\Binaries\Win64\TwitchLiveStreaming*.*
diff --git a/Engine/Plugins/Runtime/Analytics/Apsalar/IOSApsalar/Source/IOSApsalar/IOSApsalar.Build.cs b/Engine/Plugins/Runtime/Analytics/Apsalar/IOSApsalar/Source/IOSApsalar/IOSApsalar.Build.cs
--- a/Engine/Plugins/Runtime/Analytics/Apsalar/IOSApsalar/Source/IOSApsalar/IOSApsalar.Build.cs
+++ b/Engine/Plugins/Runtime/Analytics/Apsalar/IOSApsalar/Source/IOSApsalar/IOSApsalar.Build.cs
@@ -6,6 +6,8 @@ namespace UnrealBuildTool.Rules
 	{
 		public IOSApsalar( TargetInfo Target )
 		{
+			BinariesSubFolder = "NotForLicensees";
+
 			PublicDependencyModuleNames.AddRange(
 				new string[]
 				{
