Added MacClientTargetPlatform to the JunkManifest list for Win64

diff --git a/Engine/Build/JunkManifest.txt b/Engine/Build/JunkManifest.txt
--- a/Engine/Build/JunkManifest.txt
+++ b/Engine/Build/JunkManifest.txt
@@ -37,6 +37,7 @@ Engine\Binaries\Mac\UE4Editor.app\Contents\MacOS\*-CameraEditor.*
 Engine\Binaries\Mac\UE4Editor.app\Contents\MacOS\*-SessionLauncher.*
 Engine\Binaries\Mac\*-MessagingDebugger.*
 Engine\Binaries\Win64\*-DDCUtils*.*
+Engine\Binaries\Win64\*-MacClientTargetPlatform*.*
 Engine\Binaries\Win64\*-MessagingDebugger.*
 Engine\Binaries\Win64\*-OnlineSubsystemMcp*.*
 Engine\Binaries\Win64\*-SimplygonMeshReduction*.*
