Add placeholder sections to strip symbols for Android and Linux.

diff --git a/Engine/Build/InstalledEngineFilters.ini b/Engine/Build/InstalledEngineFilters.ini
--- a/Engine/Build/InstalledEngineFilters.ini
+++ b/Engine/Build/InstalledEngineFilters.ini
@@ -115,12 +115,19 @@
 ; Build files
 /Engine/Build/Android/...
 /Engine/Source/ThirdParty/Android/cxa_demangle/*
+-/Engine/Binaries/Android/....apk
 {Win64} /Engine/Binaries/ThirdParty/ARM/Win32/...
 
 ; TADP installer
 {Win64} /Engine/Extras/Android/*.exe
 {Mac} /Engine/Extras/Android/*.dmg
 
+[StripTargetPlatform.Android]
+/Engine/Binaries/Android/....a
+/Engine/Binaries/Android/....so
+/Engine/Plugins/.../Binaries/Android/....a
+/Engine/Plugins/.../Binaries/Android/....so
+
 [CopyTargetPlatform.IOS]
 ; Build files
 /Engine/Build/IOS/...
@@ -161,10 +168,12 @@
 /Engine/Binaries/ThirdParty/Steamworks/Steamv132/Win32/*.dll
 
 [StripTargetPlatform.Win32]
-/Engine/Binaries/.../*.pdb
+/Engine/Binaries/Win32/....pdb
+/Engine/Plugins/.../Binaries/Win32/....pdb
 
 [StripTargetPlatform.Win64]
-/Engine/Binaries/.../*.pdb
+/Engine/Binaries/Win64/....pdb
+/Engine/Plugins/.../Binaries/Win64/....pdb
 
 [CopyTargetPlatform.Linux]
 ; Content-only: include pre-compiled executables manually
@@ -175,3 +184,6 @@
 /Engine/Config/Controller.vdf
 /Engine/Binaries/ThirdParty/OpenAL/Linux/*
 /Engine/Binaries/ThirdParty/Steamworks/Steamv132/Linux/*.so
+
+[StripTargetPlatform.Linux]
+/Engine/Binaries/Linux/*.
diff --git a/Engine/Source/Programs/AutomationTool/Scripts/RocketBuild.Automation.cs b/Engine/Source/Programs/AutomationTool/Scripts/RocketBuild.Automation.cs
--- a/Engine/Source/Programs/AutomationTool/Scripts/RocketBuild.Automation.cs
+++ b/Engine/Source/Programs/AutomationTool/Scripts/RocketBuild.Automation.cs
@@ -449,6 +449,22 @@ namespace Rocket
 					CommandUtils.RenameFile(TempFileName, FileName);
 				}
 			}
+			else if(TargetPlatform == UnrealTargetPlatform.Android)
+			{
+				CommandUtils.LogWarning("Stripping has not yet been implemented for Android. Skipping files:");
+				foreach (string FileName in FileNames)
+				{
+					CommandUtils.LogWarning("  {0}", FileName);
+				}
+			}
+			else if(TargetPlatform == UnrealTargetPlatform.Linux)
+			{
+				CommandUtils.LogWarning("Stripping has not yet been implemented for Linux. Skipping files:");
+				foreach (string FileName in FileNames)
+				{
+					CommandUtils.LogWarning("  {0}", FileName);
+				}
+			}
 			else
 			{
 				if(FileNames.Any())
