Remove task browser.  Not used internally, not supported

diff --git a/Engine/Build/JunkManifest.txt b/Engine/Build/JunkManifest.txt
--- a/Engine/Build/JunkManifest.txt
+++ b/Engine/Build/JunkManifest.txt
@@ -17,6 +17,7 @@ Platform=Linux:Engine\Binaries\Win64\Linux\*.*
 Engine\Binaries\Win64\UE4Editor-LinuxTarget*.*
 Engine\Binaries\Win64\UE4Editor-AndroidTarget*.*
 Engine\Binaries\Win64\UE4Editor-IOSPlatformEditor*.*
+Engine\Binaries\Win64\UE4Editor-TaskBrowser*.*
 Engine\Binaries\Win64\*-ProjectWorkflows.*
 Engine\Binaries\Win64\*-WorkflowEditor.*
 Engine\Binaries\Win64\*-Workflows.*
@@ -35,6 +36,7 @@ Engine\Binaries\Mac\UE4Editor.app\Contents\Frameworks\*-SlateStyle.*
 Engine\Binaries\Mac\UE4Editor.app\Contents\MacOS\*-CameraDebugger.*
 Engine\Binaries\Mac\UE4Editor.app\Contents\MacOS\*-CameraEditor.*
 Engine\Binaries\Mac\UE4Editor.app\Contents\MacOS\*-SessionLauncher.*
+Engine\Binaries\Mac\UE4Editor.app\Contents\MacOS\*-TaskBrowser.*
 Engine\Binaries\Mac\*-MessagingDebugger.*
 Engine\Binaries\Win64\*-DDCUtils*.*
 Engine\Binaries\Win64\*-MessagingDebugger.*
diff --git a/Engine/Source/Editor/LevelEditor/LevelEditor.Build.cs b/Engine/Source/Editor/LevelEditor/LevelEditor.Build.cs
--- a/Engine/Source/Editor/LevelEditor/LevelEditor.Build.cs
+++ b/Engine/Source/Editor/LevelEditor/LevelEditor.Build.cs
@@ -74,7 +74,6 @@ public class LevelEditor : ModuleRules
 				"Layers",
 				"Levels",
                 "WorldBrowser",
-				"TaskBrowser",
 				"EditorWidgets",
 				"AssetTools",
 				"WorkspaceMenuStructure",
diff --git a/Engine/Source/Editor/TaskBrowser/Private/STaskBrowser.cpp b/Engine/Source/Editor/TaskBrowser/Private/STaskBrowser.cpp
deleted file mode 100644
diff --git a/Engine/Source/Editor/TaskBrowser/Private/STaskBrowser.h b/Engine/Source/Editor/TaskBrowser/Private/STaskBrowser.h
deleted file mode 100644
diff --git a/Engine/Source/Editor/TaskBrowser/Private/TaskBrowserModule.cpp b/Engine/Source/Editor/TaskBrowser/Private/TaskBrowserModule.cpp
deleted file mode 100644
diff --git a/Engine/Source/Editor/TaskBrowser/Private/TaskBrowserPrivatePCH.h b/Engine/Source/Editor/TaskBrowser/Private/TaskBrowserPrivatePCH.h
deleted file mode 100644
diff --git a/Engine/Source/Editor/TaskBrowser/Public/TaskBrowserModule.h b/Engine/Source/Editor/TaskBrowser/Public/TaskBrowserModule.h
deleted file mode 100644
diff --git a/Engine/Source/Editor/TaskBrowser/TaskBrowser.Build.cs b/Engine/Source/Editor/TaskBrowser/TaskBrowser.Build.cs
deleted file mode 100644
diff --git a/Engine/Source/Editor/UnrealEd/UnrealEd.Build.cs b/Engine/Source/Editor/UnrealEd/UnrealEd.Build.cs
--- a/Engine/Source/Editor/UnrealEd/UnrealEd.Build.cs
+++ b/Engine/Source/Editor/UnrealEd/UnrealEd.Build.cs
@@ -160,7 +160,6 @@ public class UnrealEd : ModuleRules
 				"DeviceManager",
 				"SettingsEditor",
 				"SessionFrontend",
-				"TaskBrowser",
 				"Sequencer",
 				"SoundClassEditor",
 				"GeometryMode",
