Enable leap motion plugin by default, and include it in the binary release (Also fixes a crash running the cooker with the plugin creator plugin enabled).

diff --git a/Engine/Build/InstalledEngineFilters.ini b/Engine/Build/InstalledEngineFilters.ini
--- a/Engine/Build/InstalledEngineFilters.ini
+++ b/Engine/Build/InstalledEngineFilters.ini
@@ -58,6 +58,8 @@
 {Win64} /Engine/Binaries/ThirdParty/CEF3/Win64/...
 {Win64} /Engine/Binaries/ThirdParty/ICU/icu4c-53_1/Win32/VS2013/*
 {Win64} /Engine/Binaries/ThirdParty/ICU/icu4c-53_1/Win64/VS2013/*
+{Win64} /Engine/Binaries/ThirdParty/Leap/Win32/Leap.dll
+{Win64} /Engine/Binaries/ThirdParty/Leap/Win64/Leap.dll
 {Win64} /Engine/Binaries/ThirdParty/nvTextureTools/Win64/*.dll
 {Win64} /Engine/Binaries/ThirdParty/Oculus/Audio/Win64/ovraudio64.dll
 {Win64} /Engine/Binaries/ThirdParty/Ogg/Win64/VS2013/*.dll
diff --git a/Engine/Plugins/Editor/PluginCreator/PluginCreator.uplugin b/Engine/Plugins/Editor/PluginCreator/PluginCreator.uplugin
--- a/Engine/Plugins/Editor/PluginCreator/PluginCreator.uplugin
+++ b/Engine/Plugins/Editor/PluginCreator/PluginCreator.uplugin
@@ -15,7 +15,8 @@
 	[
 		{
 			"Name" : "PluginCreator",
-			"Type" : "Editor"
+            "Type": "Editor",
+            "LoadingPhase": "PostEngineInit"
 		}
 	]
 }
\ No newline at end of file
diff --git a/Engine/Plugins/Runtime/LeapMotionController/LeapMotionController.uplugin b/Engine/Plugins/Runtime/LeapMotionController/LeapMotionController.uplugin
--- a/Engine/Plugins/Runtime/LeapMotionController/LeapMotionController.uplugin
+++ b/Engine/Plugins/Runtime/LeapMotionController/LeapMotionController.uplugin
@@ -9,7 +9,7 @@
 	"EngineVersion" : "4.2.0",
 	"Description" : "Support for the Leap Motion input device",
 	"Category" : "Input Devices",
-	"EnabledByDefault" : false,
+	"EnabledByDefault" : true,
 	
 	"Modules" :
 	[
