Only disable symbol stripping for Win64 editor PDBs (no plugins, no target platforms).

diff --git a/Engine/Build/InstalledEngineFilters.ini b/Engine/Build/InstalledEngineFilters.ini
--- a/Engine/Build/InstalledEngineFilters.ini
+++ b/Engine/Build/InstalledEngineFilters.ini
@@ -203,7 +203,7 @@
 [StripSymbols.Win64]
 /Engine/Binaries/Win64/....pdb
 /Engine/Plugins/.../Binaries/Win64/....pdb
-{Win64} -UE4Editor*.pdb
+{Win64} -/Engine/Binaries/Win64/UE4Editor*.pdb
 
 [StripSymbols.Linux]
 /Engine/Binaries/Linux/*.
