Remove Windows tools for compile build products from Mac, and intermediates that shouldn't be included in the build.

diff --git a/Engine/Build/InstalledEngineFilters.ini b/Engine/Build/InstalledEngineFilters.ini
--- a/Engine/Build/InstalledEngineFilters.ini
+++ b/Engine/Build/InstalledEngineFilters.ini
@@ -41,6 +41,9 @@
 ; Shaders
 /Engine/Shaders/...
 
+; Exclude all the intermediate Rocket files in the Engine/Saved folder
+-/Engine/Saved/...
+
 ; Valid templates will be included by name from the build script, but we need to exclude intermediates
 -/Templates/*/Binaries/...
 -/Templates/*/Build/Receipts/...
@@ -114,6 +117,10 @@
 
 {Mac} /Engine/Source/UE4Game*.cs
 
+; Don't want these folders, even if they're part of ToolsForCompile
+{Mac} /Engine/Binaries/Win32/...
+{Mac} /Engine/Binaries/Win64/...
+
 [CopyTargetPlatforms]
 ; API docs
 /Engine/Documentation/CHM/API.chm
