Fix direction of rule for Windows build products on Mac.

diff --git a/Engine/Build/InstalledEngineFilters.ini b/Engine/Build/InstalledEngineFilters.ini
--- a/Engine/Build/InstalledEngineFilters.ini
+++ b/Engine/Build/InstalledEngineFilters.ini
@@ -118,8 +118,8 @@
 {Mac} /Engine/Source/UE4Game*.cs
 
 ; Don't want these folders, even if they're part of ToolsForCompile
-{Mac} /Engine/Binaries/Win32/...
-{Mac} /Engine/Binaries/Win64/...
+{Mac} -/Engine/Binaries/Win32/...
+{Mac} -/Engine/Binaries/Win64/...
 
 [CopyTargetPlatforms]
 ; API docs
