[INTEGRATE] Change 2532650 by [REDACTED_PERSON][REDACTED_HANDLE] on 2015/04/30 17:39:19

	Include plugin resource and template folders in binary release. Fixes missing icons, and plugin creator not working.

diff --git a/Engine/Build/InstalledEngineFilters.ini b/Engine/Build/InstalledEngineFilters.ini
--- a/Engine/Build/InstalledEngineFilters.ini
+++ b/Engine/Build/InstalledEngineFilters.ini
@@ -28,7 +28,9 @@
 ; Plugins
 /Engine/Plugins/....uplugin
 /Engine/Plugins/.../Content/...
+/Engine/Plugins/.../Resources/...
 /Engine/Plugins/.../Source/...
+/Engine/Plugins/.../Templates/...
 -/Engine/Plugins/Runtime/TwitchLiveStreaming/...
 
 ; Source code
