addition of Target.cs.template to binary release
UE-14671

diff --git a/Engine/Build/InstalledEngineFilters.ini b/Engine/Build/InstalledEngineFilters.ini
--- a/Engine/Build/InstalledEngineFilters.ini
+++ b/Engine/Build/InstalledEngineFilters.ini
@@ -127,6 +127,9 @@
 ; API docs
 /Engine/Documentation/CHM/API.chm
 
+; Build files
+/Engine/Build/Target.cs.template
+
 ; Starter content
 /Samples/StarterContent/Content/...
 /Samples/MobileStarterContent/Content/...
