Add exclusion to ensure we dont include Mac dSYM resources in Rocket on Windows, which blows out MAX_PATH

diff --git a/Engine/Build/InstalledEngineFilters.ini b/Engine/Build/InstalledEngineFilters.ini
--- a/Engine/Build/InstalledEngineFilters.ini
+++ b/Engine/Build/InstalledEngineFilters.ini
@@ -33,6 +33,9 @@
 /Engine/Plugins/.../Templates/...
 -/Engine/Plugins/Runtime/TwitchLiveStreaming/...
 
+; Exclude Mac binaries on windows here, because previous wildcard for plugin resources matches OSX dsym resources on Windows
+{Win64} -/Engine/Plugins/.../Binaries/Mac/...
+
 ; Source code
 /Engine/Source/Runtime/...
 /Engine/Source/Developer/...
