Fixed missing PhysX DLLs in Rocket builds.

diff --git a/Engine/Build/InstalledEngineFilters.ini b/Engine/Build/InstalledEngineFilters.ini
--- a/Engine/Build/InstalledEngineFilters.ini
+++ b/Engine/Build/InstalledEngineFilters.ini
@@ -73,8 +73,9 @@
 {Win64} /Engine/Binaries/ThirdParty/Ogg/Win64/VS2013/*.dll
 {Win64} /Engine/Binaries/ThirdParty/Vorbis/Win64/VS2013/*.dll
 {Win64} /Engine/Binaries/ThirdParty/PhysX/APEX-1.3/Win64/VS2013/*PROFILE_x64.dll
-{Win64} /Engine/Binaries/ThirdParty/PhysX/PhysX-3.3/Win64/VS2013/*PROFILE_x64.dll
-{Win64} /Engine/Binaries/ThirdParty/PhysX/PhysX-3.3/Win64/VS2013/nvToolsExt64_1.dll
+{Win64} /Engine/Binaries/ThirdParty/PhysX/PhysX-3.3/Win64/VS2013/*.dll
+{Win64} -/Engine/Binaries/ThirdParty/PhysX/PhysX-3.3/Win64/VS2013/*DEBUG_x64.dll
+{Win64} -/Engine/Binaries/ThirdParty/PhysX/PhysX-3.3/Win64/VS2013/*CHECKED_x64.dll
 {Win64} /Engine/Binaries/ThirdParty/Perforce/p4api.dll
 {Win64} /Engine/Binaries/ThirdParty/Perforce/p4api.xml
 {Win64} /Engine/Binaries/ThirdParty/Steamworks/Steamv132/Win64/*.dll
@@ -183,8 +184,9 @@
 /Engine/Binaries/ThirdParty/Ogg/Win32/VS2013/*.dll
 /Engine/Binaries/ThirdParty/Vorbis/Win32/VS2013/*.dll
 /Engine/Binaries/ThirdParty/PhysX/APEX-1.3/Win32/VS2013/*PROFILE_x86.dll
-/Engine/Binaries/ThirdParty/PhysX/PhysX-3.3/Win32/VS2013/*PROFILE_x86.dll
-/Engine/Binaries/ThirdParty/PhysX/PhysX-3.3/Win32/VS2013/nvToolsExt32_1.dll
+/Engine/Binaries/ThirdParty/PhysX/PhysX-3.3/Win32/VS2013/*.dll
+-/Engine/Binaries/ThirdParty/PhysX/PhysX-3.3/Win32/VS2013/*DEBUG_x86.dll
+-/Engine/Binaries/ThirdParty/PhysX/PhysX-3.3/Win32/VS2013/*CHECKED_x86.dll
 /Engine/Binaries/ThirdParty/Steamworks/Steamv132/Win32/*.dll
 
 [CopyTargetPlatform.Linux]
