[INTEGRATE] Change 2534925 by [REDACTED_PERSON][REDACTED_HANDLE] on 2015/05/02 14:03:21

	Exclude static libraries for iOS on Windows (adding 3gb to install size), and fix case for list of files to strip symbols from on Mac.

diff --git a/Engine/Build/InstalledEngineFilters.ini b/Engine/Build/InstalledEngineFilters.ini
--- a/Engine/Build/InstalledEngineFilters.ini
+++ b/Engine/Build/InstalledEngineFilters.ini
@@ -157,10 +157,12 @@
 -/Engine/Build/IOS/....mobileprovision
 -/Engine/Build/IOS/UnrealRemoteTool
 
-; Content-only on Win64: include pre-compiled executables manually
-{Win64} /Engine/Binaries/IOS/UnrealLaunchDaemon*
-{Win64} /Engine/Binaries/IOS/UE4Game.stub
-{Win64} /Engine/Binaries/IOS/UE4Game-IOS-Shipping.stub
+; Content-only on Win64: exclude things we only need for code projects
+{Win64} -/Engine/Binaries/IOS/*.a
+{Win64} -/Engine/Binaries/IOS/UE4Game
+{Win64} -/Engine/Binaries/IOS/UE4Game-IOS-Shipping
+{Win64} -/Engine/Intermediate/.../IOS/...
+{Win64} -/Engine/Plugins/.../Binaries/IOS/...
 {Win64} /Engine/Binaries/DotNET/IOS/openssl.exe
 {Win64} /Engine/Binaries/ThirdParty/IOS/*
 {Win64} /Engine/Binaries/ThirdParty/ARM/Win32/...
@@ -205,9 +207,9 @@
 /Engine/Plugins/.../Binaries/Android/....a
 /Engine/Plugins/.../Binaries/Android/....so
 
-[StripSymbols.iOS]
-/Engine/Binaries/iOS/....a
-/Engine/Plugins/.../Binaries/iOS/....a
+[StripSymbols.IOS]
+/Engine/Binaries/IOS/....a
+/Engine/Plugins/.../Binaries/IOS/....a
 
 [StripSymbols.Mac]
 /Engine/Binaries/Mac/....a
