Add OpenVR to installed engine distribution.

diff --git a/Engine/Build/InstalledEngineFilters.ini b/Engine/Build/InstalledEngineFilters.ini
--- a/Engine/Build/InstalledEngineFilters.ini
+++ b/Engine/Build/InstalledEngineFilters.ini
@@ -82,6 +82,7 @@
 {Win64} /Engine/Binaries/ThirdParty/svn/Win64/...
 {Win64} /Engine/Binaries/ThirdParty/OpenSSL/Win64/VS2012/*.dll
 {Win64} /Engine/Binaries/ThirdParty/OpenSSL/Win64/VS2013/*.dll
+{Win64} /Engine/Binaries/ThirdParty/OpenVR/*/win64/...
 
 {Win64} /Engine/Build/BatchFiles/Build.bat
 {Win64} /Engine/Build/BatchFiles/Clean.bat
@@ -101,6 +102,7 @@
 {Mac} /Engine/Binaries/ThirdParty/CEF3/Mac/...
 {Mac} /Engine/Binaries/ThirdParty/ICU/icu4c-53_1/Mac/*
 {Mac} /Engine/Binaries/ThirdParty/ImgTec/PVRTexToolCLI
+{Mac} /Engine/Binaries/ThirdParty/OpenVR/*/osx32/...
 
 {Mac} /Engine/Binaries/ThirdParty/svn/Mac/...
 {Mac} /Engine/Binaries/ThirdParty/Mono/Mac/...
@@ -190,6 +192,7 @@
 -/Engine/Binaries/ThirdParty/PhysX/PhysX-3.3/Win32/VS2013/*DEBUG_x86.dll
 -/Engine/Binaries/ThirdParty/PhysX/PhysX-3.3/Win32/VS2013/*CHECKED_x86.dll
 /Engine/Binaries/ThirdParty/Steamworks/Steamv132/Win32/*.dll
+/Engine/Binaries/ThirdParty/OpenVR/*/win32/...
 
 [CopyTargetPlatform.Linux]
 ; Content-only: include pre-compiled executables manually
