fix for missing icons in binary build
UE-14679
#ios

diff --git a/Engine/Build/InstalledEngineFilters.ini b/Engine/Build/InstalledEngineFilters.ini
--- a/Engine/Build/InstalledEngineFilters.ini
+++ b/Engine/Build/InstalledEngineFilters.ini
@@ -154,7 +154,6 @@
 [CopyTargetPlatform.IOS]
 ; Build files
 /Engine/Build/IOS/...
--/Engine/Build/IOS/....png
 -/Engine/Build/IOS/....psd
 -/Engine/Build/IOS/....mobileprovision
 -/Engine/Build/IOS/UnrealRemoteTool
