- Added HTML5 helper app to engine distribution [UE-14958]

diff --git a/Engine/Build/InstalledEngineFilters.ini b/Engine/Build/InstalledEngineFilters.ini
--- a/Engine/Build/InstalledEngineFilters.ini
+++ b/Engine/Build/InstalledEngineFilters.ini
@@ -181,6 +181,7 @@
 [CopyTargetPlatform.HTML5]
 ; Build files
 /Engine/Build/HTML5/...
+/Engine/Binaries/DotNET/HTML5LaunchHelper.exe
 
 [CopyTargetPlatform.Win32]
 ; Just third party dependencies.
