remove html5 helper app to engine distribution. Should be built & included by build machines now.

diff --git a/Engine/Build/InstalledEngineFilters.ini b/Engine/Build/InstalledEngineFilters.ini
--- a/Engine/Build/InstalledEngineFilters.ini
+++ b/Engine/Build/InstalledEngineFilters.ini
@@ -181,7 +181,8 @@
 [CopyTargetPlatform.HTML5]
 ; Build files
 /Engine/Build/HTML5/...
-/Engine/Binaries/DotNET/HTML5LaunchHelper.exe
+; Removed HTML5LaunchHelper. It should be built by the build machines now.
+; /Engine/Binaries/DotNET/HTML5LaunchHelper.exe
 
 [CopyTargetPlatform.Win32]
 ; Just third party dependencies.
