Facebook frameworks added to the binary release for mac.

diff --git a/Engine/Build/InstalledEngineFilters.ini b/Engine/Build/InstalledEngineFilters.ini
--- a/Engine/Build/InstalledEngineFilters.ini
+++ b/Engine/Build/InstalledEngineFilters.ini
@@ -177,6 +177,7 @@
 {Mac} /Engine/Binaries/DotNET/iOS/iPhonePackager.exe.config
 {Mac} /Engine/Binaries/DotNET/iOS/MobileDeviceInterface.dll
 {Mac} /Engine/Binaries/DotNET/iOS/MobileDeviceInterface.dll.config
+{Mac} /Engine/Source/ThirdParty/Facebook/IOS/...
 
 [CopyTargetPlatform.HTML5]
 ; Build files
