- Added ASTC encoders for non-Windows desktop platforms [UE-15278]

diff --git a/Engine/Build/InstalledEngineFilters.ini b/Engine/Build/InstalledEngineFilters.ini
--- a/Engine/Build/InstalledEngineFilters.ini
+++ b/Engine/Build/InstalledEngineFilters.ini
@@ -146,6 +146,8 @@
 /Engine/Source/ThirdParty/Android/cxa_demangle/...
 -/Engine/Binaries/Android/....apk
 {Win64} /Engine/Binaries/ThirdParty/ARM/Win32/...
+{Mac} /Engine/Binaries/ThirdParty/ARM/Mac/...
+{Linux} /Engine/Binaries/ThirdParty/ARM/Linux32/...
 
 ; TADP installer
 {Win64} /Engine/Extras/Android/*.exe
