Remove some stuff from the 4.8 build that shouldn't be there, based on diffing against 4.7

diff --git a/Engine/Build/InstalledEngineFilters.ini b/Engine/Build/InstalledEngineFilters.ini
--- a/Engine/Build/InstalledEngineFilters.ini
+++ b/Engine/Build/InstalledEngineFilters.ini
@@ -63,6 +63,8 @@
 ; Win64-specific dependencies
 {Win64} /Engine/Binaries/ThirdParty/ImgTec/PVRTexToolCLI.exe
 
+{Win64} -/Engine/Binaries/Win64/UE4EditorServices.*
+{Win64} -/Engine/Binaries/Win64/UnrealVersionSelector*
 {Win64} /Engine/Binaries/Win64/libfbxsdk.dll
 {Win64} /Engine/Binaries/ThirdParty/CEF3/Win32/...
 {Win64} /Engine/Binaries/ThirdParty/CEF3/Win64/...
@@ -92,6 +94,11 @@
 {Win64} /Engine/Build/BatchFiles/RocketGenerateProjectFiles.bat
 {Win64} /Engine/Build/BatchFiles/RunUAT.bat
 {Win64} /Engine/Build/BatchFiles/GetVSComnToolsPath.bat
+{Win64} -/Engine/Build/Receipts/PS4MapFile*
+{Win64} -/Engine/Build/Receipts/XboxOnePDBFileUtil*
+{Win64} -/Engine/Build/Receipts/UE4Client-Linux*
+{Win64} -/Engine/Build/Receipts/UE4EditorServices*
+{Win64} -/Engine/Build/Receipts/UE4Server-Linux*
 
 {Win64} /Engine/Extras/3dsMaxScripts/...
 {Win64} /Engine/Extras/VisualStudioDebugging/UE4.natvis
@@ -162,6 +169,7 @@
 
 ; Content-only on Win64: exclude things we only need for code projects
 {Win64} -/Engine/Binaries/IOS/*.a
+{Win64} -/Engine/Binaries/IOS/.../*.a
 {Win64} -/Engine/Binaries/IOS/UE4Game
 {Win64} -/Engine/Binaries/IOS/UE4Game-IOS-Shipping
 {Win64} -/Engine/Intermediate/.../IOS/...
@@ -200,6 +208,7 @@
 
 [CopyTargetPlatform.Linux]
 ; Content-only: include pre-compiled executables manually
+-/Engine/Binaries/Linux/*
 /Engine/Binaries/Linux/CrashReportClient
 /Engine/Binaries/Linux/UE4Game
 /Engine/Binaries/Linux/UE4Game-Linux-Shipping
