Don't strip private symbols for Paper2D in Rocket builds.

diff --git a/Engine/Build/InstalledEngineFilters.ini b/Engine/Build/InstalledEngineFilters.ini
--- a/Engine/Build/InstalledEngineFilters.ini
+++ b/Engine/Build/InstalledEngineFilters.ini
@@ -239,6 +239,7 @@
 /Engine/Binaries/Win64/....pdb
 /Engine/Plugins/.../Binaries/Win64/....pdb
 {Win64} -/Engine/Binaries/Win64/UE4Editor*.pdb
+{Win64} -/Engine/Plugins/2D/Paper2D/Binaries/Win64/UE4Editor*.pdb
 
 [StripSymbols.Linux]
 /Engine/Binaries/Linux/*.
