Including localization data files for the engine/editor. Removing .po files.

diff --git a/Engine/Build/InstalledEngineFilters.ini b/Engine/Build/InstalledEngineFilters.ini
--- a/Engine/Build/InstalledEngineFilters.ini
+++ b/Engine/Build/InstalledEngineFilters.ini
@@ -15,8 +15,7 @@
 -/Engine/Content/....psd
 -/Engine/Content/....pdn
 -/Engine/Content/....fbx
--/Engine/Content/....archive
--/Engine/Content/....manifest
+-/Engine/Content/....po
 
 ; Config files
 /Engine/Config/...
