Add deprecated DesktopTargetPlatform to junk manifest on Linux as well.

diff --git a/Engine/Build/JunkManifest.txt b/Engine/Build/JunkManifest.txt
--- a/Engine/Build/JunkManifest.txt
+++ b/Engine/Build/JunkManifest.txt
@@ -75,3 +75,4 @@ Engine\Binaries\Win64\Android\*Android_ES31*.*
 Engine\Binaries\Mac\UE4Editor.app\Contents\MacOS\UE4Editor-DesktopTargetPlatform.dylib
 Engine\Binaries\Mac\UnrealFrontend.app\Contents\MacOS\UnrealFrontend-DesktopTargetPlatform.dylib
 Engine\Binaries\Win64\*-DesktopTargetPlatform*.*
+Engine\Binaries\Linux\*-DesktopTargetPlatform*.*
