Paper2D: Added old name of renamed module PaperJsonImporter to the junk manifest

diff --git a/Engine/Build/JunkManifest.txt b/Engine/Build/JunkManifest.txt
--- a/Engine/Build/JunkManifest.txt
+++ b/Engine/Build/JunkManifest.txt
@@ -64,6 +64,10 @@ Engine\Plugins\Script\ScriptPlugin\Binaries\Win64\*
 Engine\Plugins\Experimental\SteamVR\Binaries\Win64\*
 Engine\Plugins\Experimental\Paper2D\Binaries
 Engine\Plugins\Experimental\Paper2D\Intermediate
+Engine\Plugins\2D\Paper2D\Binaries\Win32\*-PaperJsonImporter*.*
+Engine\Plugins\2D\Paper2D\Binaries\Win64\*-PaperJsonImporter*.*
+Engine\Plugins\2D\Paper2D\Binaries\Mac\*-PaperJsonImporter*.*
+Engine\Plugins\2D\Paper2D\Binaries\Linux\*-PaperJsonImporter*.*
 Engine\Plugins\Runtime\Analytics\Apsalar\IOSApsalar\Binaries\Mac\*-IOSApsalar.*
 Engine\Plugins\Runtime\TwitchLiveStreaming\Binaries\Win32\TwitchLiveStreaming*.*
 Engine\Plugins\Runtime\TwitchLiveStreaming\Binaries\Win32\*-TwitchLiveStreaming*.*
