UE4 - remove proceduralmeshcomponent plugin via junk manifest

diff --git a/Engine/Build/JunkManifest.txt b/Engine/Build/JunkManifest.txt
--- a/Engine/Build/JunkManifest.txt
+++ b/Engine/Build/JunkManifest.txt
@@ -80,3 +80,4 @@ Engine\Binaries\Mac\UE4Editor.app\Contents\MacOS\UE4Editor-DesktopTargetPlatform
 Engine\Binaries\Mac\UnrealFrontend.app\Contents\MacOS\UnrealFrontend-DesktopTargetPlatform.dylib
 Engine\Binaries\Win64\*-DesktopTargetPlatform*.*
 Engine\Binaries\Linux\*-DesktopTargetPlatform*.*
+Engine\Plugins\Runtime\ProceduralMeshComponent
