Include HTML5 SDK in rocket distribution

diff --git a/Engine/Build/InstalledEngineFilters.ini b/Engine/Build/InstalledEngineFilters.ini
--- a/Engine/Build/InstalledEngineFilters.ini
+++ b/Engine/Build/InstalledEngineFilters.ini
@@ -191,8 +191,10 @@
 [CopyTargetPlatform.HTML5]
 ; Build files
 /Engine/Build/HTML5/...
-; Removed HTML5LaunchHelper. It should be built by the build machines now.
-; /Engine/Binaries/DotNET/HTML5LaunchHelper.exe
+; SDK
+/Engine/Source/ThirdParty/HTML5/emsdk/emscripten/...
+{Win64} /Engine/Source/ThirdParty/HTML5/emsdk/Win64/...
+{Mac} /Engine/Source/ThirdParty/HTML5/emsdk/Mac/...
 
 [CopyTargetPlatform.Win32]
 ; Just third party dependencies.
