RocketBuild: add Linux-specific files.

diff --git a/Engine/Build/InstalledEngineFilters.ini b/Engine/Build/InstalledEngineFilters.ini
--- a/Engine/Build/InstalledEngineFilters.ini
+++ b/Engine/Build/InstalledEngineFilters.ini
@@ -135,6 +135,29 @@
 {Mac} -/Engine/Binaries/Win32/...
 {Mac} -/Engine/Binaries/Win64/...
 
+; Linux-specific dependencies
+;{Linux} /Engine/Binaries/ThirdParty/CEF3/Linux/...
+;{Linux} /Engine/Binaries/ThirdParty/ImgTec/PVRTexToolCLI
+{Linux} /Engine/Binaries/ThirdParty/ICU/icu4c-53_1/Linux/...
+{Linux} /Engine/Binaries/Linux/libnvcore.so
+{Linux} /Engine/Binaries/Linux/libnvimage.so
+{Linux} /Engine/Binaries/Linux/libnvmath.so
+{Linux} /Engine/Binaries/Linux/libnvtt.so
+{Linux} /Engine/Binaries/Linux/libmcpp.so.0
+{Linux} /Engine/Binaries/Linux/libopenal.so.1
+
+{Linux} /Engine/Build/BatchFiles/Linux/Build.sh
+{Linux} /Engine/Build/BatchFiles/Linux/BuildThirdParty.sh
+{Linux} /Engine/Build/BatchFiles/Linux/GenerateProjectFiles.sh
+{Linux} /Engine/Build/BatchFiles/Linux/GitDependencies.sh
+{Linux} /Engine/Build/BatchFiles/Linux/QASmokeManual.sh
+{Linux} /Engine/Build/BatchFiles/Linux/RunMono.sh
+{Linux} /Engine/Build/BatchFiles/Linux/RunXBuild.sh
+{Linux} /Engine/Build/BatchFiles/Linux/Setup.sh
+{Linux} /Engine/Build/BatchFiles/Linux/UpdateDeps.sh
+{Linux} /Engine/Build/BatchFiles/RunUAT.sh
+
+
 [CopyTargetPlatforms]
 ; API docs
 /Engine/Documentation/CHM/API.chm
