Reset values for build.version.

diff --git a/Engine/Build/Build.version b/Engine/Build/Build.version
--- a/Engine/Build/Build.version
+++ b/Engine/Build/Build.version
@@ -2,7 +2,7 @@
 	"MajorVersion": 4,
 	"MinorVersion": 9,
 	"PatchVersion": 0,
-	"Changelist": 2616747,
-	"IsLicenseeVersion": 1,
-	"BranchName": "++depot+UE4"
+	"Changelist": 0,
+	"IsLicenseeVersion": 0,
+	"BranchName": ""
 }
\ No newline at end of file
