remove legacey xmpp binaries

diff --git a/Engine/Build/JunkManifest.txt b/Engine/Build/JunkManifest.txt
--- a/Engine/Build/JunkManifest.txt
+++ b/Engine/Build/JunkManifest.txt
@@ -48,7 +48,6 @@ Engine\Binaries\Win64\*-OnlineSubsystemMcp*.*
 Engine\Binaries\Win64\*-SimplygonMeshReduction*.*
 Engine\Binaries\Win64\*-TextureXboxOneFormat.*
 Engine\Binaries\Win64\*-TextureFormatATITC.*
-Engine\Binaries\Win64\*-XMPP*.*
 Shadow\Binaries\Win64\*ShadowLoadingScreen.*
 Engine\Binaries\Win64\UnrealRemoteAgent.*
 Engine\Binaries\DotNET\AutomationScripts\NoRedist\FotniteGame.Automation.*
@@ -81,3 +80,4 @@ Engine\Binaries\Mac\UnrealFrontend.app\Contents\MacOS\UnrealFrontend-DesktopTarg
 Engine\Binaries\Win64\*-DesktopTargetPlatform*.*
 Engine\Binaries\Linux\*-DesktopTargetPlatform*.*
 Engine\Source\Runtime\ProceduralMeshComponent
+Engine\Source\Runtime\Online\NotForLicensees\*-XMPP*.*
