updated delete for deprecated xmpp binary

diff --git a/Engine/Build/JunkManifest.txt b/Engine/Build/JunkManifest.txt
--- a/Engine/Build/JunkManifest.txt
+++ b/Engine/Build/JunkManifest.txt
@@ -50,6 +50,7 @@ Engine\Binaries\Win64\*-TextureXboxOneFormat.*
 Engine\Binaries\Win64\*-TextureFormatATITC.*
 Shadow\Binaries\Win64\*ShadowLoadingScreen.*
 Engine\Binaries\Win64\UnrealRemoteAgent.*
+Engine\Binaries\Win64\NotForLicensees\*-XMPP*.*
 Engine\Binaries\DotNET\AutomationScripts\NoRedist\FotniteGame.Automation.*
 Engine\Plugins\Editor\IntroTutorials\Binaries\Win64\*-IntroTutorials*.*
 Engine\Plugins\Editor\IntroTutorials\Binaries\Mac\*-IntroTutorials*.*
@@ -80,4 +81,3 @@ Engine\Binaries\Mac\UnrealFrontend.app\Contents\MacOS\UnrealFrontend-DesktopTarg
 Engine\Binaries\Win64\*-DesktopTargetPlatform*.*
 Engine\Binaries\Linux\*-DesktopTargetPlatform*.*
 Engine\Source\Runtime\ProceduralMeshComponent
-Engine\Source\Runtime\Online\NotForLicensees\*-XMPP*.*
