#UE4 #GameSubCatalog #ReleaseFix-4.9 #OrionBlocker
- Removed prebuilt engine binaries for GSC from P4 (since that module moved to a game plugin)
- Added these paths to the JunkManifest per BenM

diff --git a/Engine/Build/JunkManifest.txt b/Engine/Build/JunkManifest.txt
--- a/Engine/Build/JunkManifest.txt
+++ b/Engine/Build/JunkManifest.txt
@@ -84,3 +84,5 @@ Engine\Binaries\Mac\UnrealFrontend.app\Contents\MacOS\UnrealFrontend-DesktopTarg
 Engine\Binaries\Win64\*-DesktopTargetPlatform*.*
 Engine\Binaries\Linux\*-DesktopTargetPlatform*.*
 Engine\Source\Runtime\ProceduralMeshComponent
+Engine\Binaries\Win64\*-GameSubCatalog.*
+Engine\Binaries\Mac\UE4Editor.app\Contents\MacOS\UE4Editor-GameSubCatalog.dylib
