Made shared feature pack resources get packaged into rocket build

diff --git a/Engine/Build/InstalledEngineFilters.ini b/Engine/Build/InstalledEngineFilters.ini
--- a/Engine/Build/InstalledEngineFilters.ini
+++ b/Engine/Build/InstalledEngineFilters.ini
@@ -170,6 +170,9 @@
 /Samples/StarterContent/Content/...
 /Samples/MobileStarterContent/Content/...
 
+; shared template resources
+/Templates/TemplateResources/...
+
 [CopyTargetPlatform.Android]
 ; Build files
 /Engine/Build/Android/...
