Added deprecated Mac XMPP libs to junk manifest

diff --git a/Engine/Build/JunkManifest.txt b/Engine/Build/JunkManifest.txt
--- a/Engine/Build/JunkManifest.txt
+++ b/Engine/Build/JunkManifest.txt
@@ -41,6 +41,7 @@ Engine\Binaries\Mac\UE4Editor.app\Contents\MacOS\*-CameraEditor.*
 Engine\Binaries\Mac\UE4Editor.app\Contents\MacOS\*-SessionLauncher.*
 Engine\Binaries\Mac\UE4Editor.app\Contents\MacOS\*-TaskBrowser.*
 Engine\Binaries\Mac\UE4Editor.app\Contents\MacOS\*-Levels.*
+Engine\Binaries\Mac\UE4Editor.app\Contents\MacOS\NotForLicensees\*-XMPP*.*
 Engine\Binaries\Mac\*-MessagingDebugger.*
 Engine\Binaries\Win64\*-DDCUtils*.*
 Engine\Binaries\Win64\*-MessagingDebugger.*
