Added old vrlib.jar for Android to junk manifest

diff --git a/Engine/Build/JunkManifest.txt b/Engine/Build/JunkManifest.txt
--- a/Engine/Build/JunkManifest.txt
+++ b/Engine/Build/JunkManifest.txt
@@ -53,6 +53,7 @@ Shadow\Binaries\Win64\*ShadowLoadingScreen.*
 Engine\Binaries\Win64\UnrealRemoteAgent.*
 Engine\Binaries\Win64\NotForLicensees\*-XMPP*.*
 Engine\Binaries\DotNET\AutomationScripts\NoRedist\FotniteGame.Automation.*
+Engine\Build\Android\Java\libs\vrlib.jar
 Engine\Plugins\Editor\IntroTutorials\Binaries\Win64\*-IntroTutorials*.*
 Engine\Plugins\Editor\IntroTutorials\Binaries\Mac\*-IntroTutorials*.*
 Engine\Plugins\Editor\PluginCreator\Binaries\Linux\*-PluginCreator*.*
