Added Qualcomm third party dll to Win64 dist for Android
#ue4
#android

diff --git a/Engine/Build/InstalledEngineFilters.ini b/Engine/Build/InstalledEngineFilters.ini
--- a/Engine/Build/InstalledEngineFilters.ini
+++ b/Engine/Build/InstalledEngineFilters.ini
@@ -86,6 +86,7 @@
 {Win64} /Engine/Binaries/ThirdParty/OpenSSL/Win64/VS2012/*.dll
 {Win64} /Engine/Binaries/ThirdParty/OpenSSL/Win64/VS2013/*.dll
 {Win64} /Engine/Binaries/ThirdParty/OpenVR/*/win64/...
+{Win64} /Engine/Binaries/ThirdParty/Qualcomm/Win64/TextureConverter.dll
 
 {Win64} /Engine/Build/BatchFiles/Build.bat
 {Win64} /Engine/Build/BatchFiles/Clean.bat
