Back out revision 37 from //depot/UE4/Engine/Build/InstalledEngineFilters.ini until BuildDerivedDataCacheNode can also use receipts for runtime dependencies.

diff --git a/Engine/Build/InstalledEngineFilters.ini b/Engine/Build/InstalledEngineFilters.ini
--- a/Engine/Build/InstalledEngineFilters.ini
+++ b/Engine/Build/InstalledEngineFilters.ini
@@ -60,13 +60,33 @@
 -/Templates/*/contents.txt
 
 ; Win64-specific dependencies
+{Win64} /Engine/Binaries/ThirdParty/ImgTec/PVRTexToolCLI.exe
+
 {Win64} -/Engine/Binaries/Win64/UE4EditorServices.*
 {Win64} -/Engine/Binaries/Win64/UnrealVersionSelector*
+{Win64} /Engine/Binaries/Win64/libfbxsdk.dll
+{Win64} /Engine/Binaries/ThirdParty/CEF3/Win32/...
+{Win64} /Engine/Binaries/ThirdParty/CEF3/Win64/...
 {Win64} /Engine/Binaries/ThirdParty/ICU/icu4c-53_1/Win32/VS2013/*
 {Win64} /Engine/Binaries/ThirdParty/ICU/icu4c-53_1/Win64/VS2013/*
+{Win64} /Engine/Binaries/ThirdParty/Leap/Win32/Leap.dll
+{Win64} /Engine/Binaries/ThirdParty/Leap/Win64/Leap.dll
+{Win64} /Engine/Binaries/ThirdParty/nvTextureTools/Win64/*.dll
+{Win64} /Engine/Binaries/ThirdParty/Oculus/Audio/Win64/ovraudio64.dll
+{Win64} /Engine/Binaries/ThirdParty/Ogg/Win64/VS2013/*.dll
+{Win64} /Engine/Binaries/ThirdParty/Vorbis/Win64/VS2013/*.dll
+{Win64} /Engine/Binaries/ThirdParty/PhysX/APEX-1.3/Win64/VS2013/*PROFILE_x64.dll
 {Win64} /Engine/Binaries/ThirdParty/PhysX/PhysX-3.3/Win64/VS2013/*.dll
 {Win64} -/Engine/Binaries/ThirdParty/PhysX/PhysX-3.3/Win64/VS2013/*DEBUG_x64.dll
 {Win64} -/Engine/Binaries/ThirdParty/PhysX/PhysX-3.3/Win64/VS2013/*CHECKED_x64.dll
+{Win64} /Engine/Binaries/ThirdParty/Perforce/p4api.dll
+{Win64} /Engine/Binaries/ThirdParty/Perforce/p4api.xml
+{Win64} /Engine/Binaries/ThirdParty/Steamworks/Steamv132/Win64/*.dll
+{Win64} /Engine/Binaries/ThirdParty/svn/Win64/...
+{Win64} /Engine/Binaries/ThirdParty/OpenSSL/Win64/VS2012/*.dll
+{Win64} /Engine/Binaries/ThirdParty/OpenSSL/Win64/VS2013/*.dll
+{Win64} /Engine/Binaries/ThirdParty/OpenVR/*/win64/...
+{Win64} /Engine/Binaries/ThirdParty/Qualcomm/Win64/TextureConverter.dll
 
 {Win64} /Engine/Build/BatchFiles/Build.bat
 {Win64} /Engine/Build/BatchFiles/Clean.bat
@@ -88,8 +108,14 @@
 {Win64} /Engine/Extras/Redist/en-us/*
 
 ; Mac-specific dependencies
+{Mac} /Engine/Binaries/ThirdParty/CEF3/Mac/...
 {Mac} /Engine/Binaries/ThirdParty/ICU/icu4c-53_1/Mac/*
+{Mac} /Engine/Binaries/ThirdParty/ImgTec/PVRTexToolCLI
+{Mac} /Engine/Binaries/ThirdParty/OpenVR/*/osx32/...
+
+{Mac} /Engine/Binaries/ThirdParty/svn/Mac/...
 {Mac} /Engine/Binaries/ThirdParty/Mono/Mac/...
+{Mac} /Engine/Source/ThirdParty/Mac/RadioEffectUnit/...
 
 {Mac} /Engine/Build/BatchFiles/Mac/Build.sh
 {Mac} /Engine/Build/BatchFiles/Mac/GenerateLLDBInit.sh
@@ -111,7 +137,15 @@
 {Mac} -/Engine/Binaries/Win64/...
 
 ; Linux-specific dependencies
+;{Linux} /Engine/Binaries/ThirdParty/CEF3/Linux/...
+;{Linux} /Engine/Binaries/ThirdParty/ImgTec/PVRTexToolCLI
 {Linux} /Engine/Binaries/ThirdParty/ICU/icu4c-53_1/Linux/...
+{Linux} /Engine/Binaries/Linux/libnvcore.so
+{Linux} /Engine/Binaries/Linux/libnvimage.so
+{Linux} /Engine/Binaries/Linux/libnvmath.so
+{Linux} /Engine/Binaries/Linux/libnvtt.so
+{Linux} /Engine/Binaries/Linux/libmcpp.so.0
+{Linux} /Engine/Binaries/Linux/libopenal.so.1
 
 {Linux} /Engine/Build/BatchFiles/Linux/Build.sh
 {Linux} /Engine/Build/BatchFiles/Linux/BuildThirdParty.sh
@@ -145,6 +179,9 @@
 /Engine/Build/Android/...
 /Engine/Source/ThirdParty/Android/cxa_demangle/...
 -/Engine/Binaries/Android/....apk
+{Win64} /Engine/Binaries/ThirdParty/ARM/Win32/...
+{Mac} /Engine/Binaries/ThirdParty/ARM/Mac/...
+{Linux} /Engine/Binaries/ThirdParty/ARM/Linux32/...
 
 ; AndroidWorks installer
 {Win64} /Engine/Extras/Android/*.exe
@@ -166,6 +203,7 @@
 {Win64} -/Engine/Plugins/.../Binaries/IOS/...
 {Win64} /Engine/Binaries/DotNET/IOS/openssl.exe
 {Win64} /Engine/Binaries/ThirdParty/IOS/*
+{Win64} /Engine/Binaries/ThirdParty/ARM/Win32/...
 
 ; All tools for Mac
 {Mac} /Engine/Binaries/DotNET/iOS/DeploymentServer.exe
@@ -188,9 +226,14 @@
 
 [CopyTargetPlatform.Win32]
 ; Just third party dependencies.
+/Engine/Binaries/ThirdParty/Ogg/Win32/VS2013/*.dll
+/Engine/Binaries/ThirdParty/Vorbis/Win32/VS2013/*.dll
+/Engine/Binaries/ThirdParty/PhysX/APEX-1.3/Win32/VS2013/*PROFILE_x86.dll
 /Engine/Binaries/ThirdParty/PhysX/PhysX-3.3/Win32/VS2013/*.dll
 -/Engine/Binaries/ThirdParty/PhysX/PhysX-3.3/Win32/VS2013/*DEBUG_x86.dll
 -/Engine/Binaries/ThirdParty/PhysX/PhysX-3.3/Win32/VS2013/*CHECKED_x86.dll
+/Engine/Binaries/ThirdParty/Steamworks/Steamv132/Win32/*.dll
+/Engine/Binaries/ThirdParty/OpenVR/*/win32/...
 
 [CopyTargetPlatform.Linux]
 ; Content-only: include pre-compiled executables manually
