Fixed LevelSequence module being treated as junk

This addresses UE-21617

diff --git a/Engine/Build/JunkManifest.txt b/Engine/Build/JunkManifest.txt
--- a/Engine/Build/JunkManifest.txt
+++ b/Engine/Build/JunkManifest.txt
@@ -18,7 +18,8 @@ Engine\Binaries\Win64\UE4Editor-LinuxTarget*.*
 Engine\Binaries\Win64\UE4Editor-AndroidTarget*.*
 Engine\Binaries\Win64\UE4Editor-IOSPlatformEditor*.*
 Engine\Binaries\Win64\UE4Editor-TaskBrowser*.*
-Engine\Binaries\Win64\UE4Editor-Levels*.*
+Engine\Binaries\Win64\UE4Editor-Levels.*
+Engine\Binaries\Win64\UE4Editor-Levels-*.*
 Engine\Binaries\Win64\UnrealFrontend-VisualStudioSourceCodeAccess*.*
 Engine\Binaries\Win64\*-ProjectWorkflows.*
 Engine\Binaries\Win64\*-WorkflowEditor.*
