Add tags for the optional parts of Rocket downloads.

diff --git a/Engine/Build/InstalledEngineFilters.ini b/Engine/Build/InstalledEngineFilters.ini
--- a/Engine/Build/InstalledEngineFilters.ini
+++ b/Engine/Build/InstalledEngineFilters.ini
@@ -229,3 +229,20 @@
 
 [StripSymbols.Linux]
 /Engine/Binaries/Linux/*.
+
+[LauncherAttribute.compressed]
+*.pdb
+*.lib
+
+[LauncherAttribute.readonly]
+/Engine/Source/...
+
+[LauncherAttribute.tag:pdb_files]
+*.pdb
+
+[LauncherAttribute.tag:starter_content]
+/Samples/StarterContent/...
+
+[LauncherAttribute.tag:templates]
+/FeaturePacks/...
+/Templates/...
