Include Mac dsym files in the editor_symbols Rocket category.

diff --git a/Engine/Build/InstalledEngineFilters.ini b/Engine/Build/InstalledEngineFilters.ini
--- a/Engine/Build/InstalledEngineFilters.ini
+++ b/Engine/Build/InstalledEngineFilters.ini
@@ -237,8 +237,9 @@
 [LauncherAttribute.readonly]
 /Engine/Source/...
 
-[LauncherAttribute.tag:pdb_files]
+[LauncherAttribute.tag:editor_symbols]
 *.pdb
+*.dsym
 
 [LauncherAttribute.tag:starter_content]
 /Samples/StarterContent/...
