Update paths to pull in VS2015 libraries for Rocket.

diff --git a/Engine/Build/InstalledEngineFilters.ini b/Engine/Build/InstalledEngineFilters.ini
--- a/Engine/Build/InstalledEngineFilters.ini
+++ b/Engine/Build/InstalledEngineFilters.ini
@@ -62,11 +62,11 @@
 ; Win64-specific dependencies
 {Win64} -/Engine/Binaries/Win64/UE4EditorServices.*
 {Win64} -/Engine/Binaries/Win64/UnrealVersionSelector*
-{Win64} /Engine/Binaries/ThirdParty/ICU/icu4c-53_1/Win32/VS2013/*
-{Win64} /Engine/Binaries/ThirdParty/ICU/icu4c-53_1/Win64/VS2013/*
-{Win64} /Engine/Binaries/ThirdParty/PhysX/PhysX-3.3/Win64/VS2013/*.dll
-{Win64} -/Engine/Binaries/ThirdParty/PhysX/PhysX-3.3/Win64/VS2013/*DEBUG_x64.dll
-{Win64} -/Engine/Binaries/ThirdParty/PhysX/PhysX-3.3/Win64/VS2013/*CHECKED_x64.dll
+{Win64} /Engine/Binaries/ThirdParty/ICU/icu4c-53_1/Win32/VS2015/*
+{Win64} /Engine/Binaries/ThirdParty/ICU/icu4c-53_1/Win64/VS2015/*
+{Win64} /Engine/Binaries/ThirdParty/PhysX/PhysX-3.3/Win64/VS2015/*.dll
+{Win64} -/Engine/Binaries/ThirdParty/PhysX/PhysX-3.3/Win64/VS2015/*DEBUG_x64.dll
+{Win64} -/Engine/Binaries/ThirdParty/PhysX/PhysX-3.3/Win64/VS2015/*CHECKED_x64.dll
 
 {Win64} /Engine/Build/BatchFiles/Build.bat
 {Win64} /Engine/Build/BatchFiles/Clean.bat
@@ -188,9 +188,9 @@
 
 [CopyTargetPlatform.Win32]
 ; Just third party dependencies.
-/Engine/Binaries/ThirdParty/PhysX/PhysX-3.3/Win32/VS2013/*.dll
--/Engine/Binaries/ThirdParty/PhysX/PhysX-3.3/Win32/VS2013/*DEBUG_x86.dll
--/Engine/Binaries/ThirdParty/PhysX/PhysX-3.3/Win32/VS2013/*CHECKED_x86.dll
+/Engine/Binaries/ThirdParty/PhysX/PhysX-3.3/Win32/VS2015/*.dll
+-/Engine/Binaries/ThirdParty/PhysX/PhysX-3.3/Win32/VS2015/*DEBUG_x86.dll
+-/Engine/Binaries/ThirdParty/PhysX/PhysX-3.3/Win32/VS2015/*CHECKED_x86.dll
 
 [CopyTargetPlatform.Linux]
 ; Content-only: include pre-compiled executables manually
