Update the Rocket filter for the AndroidWorks installer.

diff --git a/Engine/Build/InstalledEngineFilters.ini b/Engine/Build/InstalledEngineFilters.ini
--- a/Engine/Build/InstalledEngineFilters.ini
+++ b/Engine/Build/InstalledEngineFilters.ini
@@ -147,8 +147,8 @@
 -/Engine/Binaries/Android/....apk
 
 ; AndroidWorks installer
-{Win64} /Engine/Extras/Android/*.exe
-{Mac} /Engine/Extras/Android/*.dmg
+{Win64} /Engine/Extras/AndroidWorks/Win64/*.exe
+{Mac} /Engine/Extras/AndroidWorks/Mac/*.dmg
 
 [CopyTargetPlatform.IOS]
 ; Build files
