Merging CL#2727943 using //depot/UE4-To-//UE4/Main

Fixed build error caused by old JunkManifest rule removing OnlineSubSystemMCP dll from its new location

diff --git a/Engine/Build/JunkManifest.txt b/Engine/Build/JunkManifest.txt
--- a/Engine/Build/JunkManifest.txt
+++ b/Engine/Build/JunkManifest.txt
@@ -46,12 +46,12 @@ Engine\Binaries\Mac\UE4Editor.app\Contents\MacOS\NotForLicensees\*-XMPP*.*
 Engine\Binaries\Mac\*-MessagingDebugger.*
 Engine\Binaries\Win64\*-DDCUtils*.*
 Engine\Binaries\Win64\*-MessagingDebugger.*
-Engine\Binaries\Win64\*-OnlineSubsystemMcp*.*
 Engine\Binaries\Win64\*-SimplygonMeshReduction*.*
 Engine\Binaries\Win64\*-TextureXboxOneFormat.*
 Engine\Binaries\Win64\*-TextureFormatATITC.*
 Shadow\Binaries\Win64\*ShadowLoadingScreen.*
 Engine\Binaries\Win64\UnrealRemoteAgent.*
+Engine\Binaries\Win64\NotForLicensees\*-OnlineSubsystemMcp*.*
 Engine\Binaries\Win64\NotForLicensees\*-XMPP*.*
 Engine\Binaries\DotNET\AutomationScripts\NoRedist\FotniteGame.Automation.*
 Engine\Build\Android\Java\libs\vrlib.jar
