Use 'Engine' subfolder for Rocket builds, to match the directory layout for every other chunked product.

diff --git a/Engine/Build/InstalledEngineBuild.xml b/Engine/Build/InstalledEngineBuild.xml
--- a/Engine/Build/InstalledEngineBuild.xml
+++ b/Engine/Build/InstalledEngineBuild.xml
@@ -13,8 +13,8 @@
 	<Property Name="BuildLabel" Default=""/>
 	<Property Name="LabeledBuildDir" Value="$(BuildLabel)"/>
 	<Property Name="LabeledBuildDir" Value="$(LabeledBuildDir)/" If="'$(LabeledBuildDir)' != '' And !HasTrailingSlash('$(LabeledBuildDir)')"/>
-	<Property Name="PublishedInstallDir" Value="$(OutputDir)/$(LabeledBuildDir)Windows"/>
-	<Property Name="PublishedInstallDirMac" Value="$(OutputDirMac)/$(LabeledBuildDir)Mac"/>
+	<Property Name="PublishedInstallDir" Value="$(OutputDir)/Engine/$(LabeledBuildDir)Windows"/>
+	<Property Name="PublishedInstallDirMac" Value="$(OutputDirMac)/Engine/$(LabeledBuildDir)Mac"/>
 
 	<!-- Directory for storing build products like the Compressed DDC, Stripped and Signed files -->
 	<Property Name="SavedOutput" Value="$(RootDir)/Saved"/>
diff --git a/Engine/Source/Programs/AutomationTool/Scripts/RocketBuild.Automation.cs b/Engine/Source/Programs/AutomationTool/Scripts/RocketBuild.Automation.cs
--- a/Engine/Source/Programs/AutomationTool/Scripts/RocketBuild.Automation.cs
+++ b/Engine/Source/Programs/AutomationTool/Scripts/RocketBuild.Automation.cs
@@ -157,7 +157,7 @@ namespace Rocket
 				string PublishedEngineDir;
 				if (ShouldDoSeriousThingsLikeP4CheckinAndPostToMCP(BranchConfig))
 				{
-					PublishedEngineDir = CommandUtils.CombinePaths(CommandUtils.RootBuildStorageDirectory(), "Rocket", "Automated", GetBuildLabel(), CommandUtils.GetGenericPlatformName(HostPlatform));
+					PublishedEngineDir = CommandUtils.CombinePaths(CommandUtils.RootBuildStorageDirectory(), "Rocket", "Automated", "Engine", GetBuildLabel(), CommandUtils.GetGenericPlatformName(HostPlatform));
 				}
 				else
 				{
