Add moved modules to JunkManifest.txt

diff --git a/Engine/Build/JunkManifest.txt b/Engine/Build/JunkManifest.txt
--- a/Engine/Build/JunkManifest.txt
+++ b/Engine/Build/JunkManifest.txt
@@ -29,6 +29,12 @@ Engine\Binaries\Win64\*-CameraEditor*.*
 Engine\Binaries\Win64\*-CameraDebugger*.*
 Engine\Binaries\Win64\*-SessionLauncher*.*
 Engine\Binaries\Win64\*-OnlineSubsystem*.*
+Engine\Binaries\Win64\*-PurchaseFlow.*
+Engine\Binaries\Win64\*-Hotfix.*
+Engine\Binaries\Win64\*-Lobby.*
+Engine\Binaries\Win64\*-Party.*
+Engine\Binaries\Win64\*-Qos.*
+Engine\Binaries\Win64\*-OnlineBlueprintSupport.*
 Engine\Binaries\Mac\UnrealFrontend.app\Contents\Frameworks\*-CoreSlateStyle.*
 Engine\Binaries\Mac\UnrealFrontend.app\Contents\Frameworks\*-SlateStyle.*
 Engine\Binaries\Mac\UnrealFrontend.app\Contents\MacOS\*-XCodeSourceCodeAccess.*
