Make sure we don't try to sign third party dlls
Check whether referenced pdbs marked as copy local exist in the output dir instead of where they originally were - avoids adding files that don't exist when using enumerate only option

diff --git a/Engine/Build/InstalledEngineBuild.xml b/Engine/Build/InstalledEngineBuild.xml
--- a/Engine/Build/InstalledEngineBuild.xml
+++ b/Engine/Build/InstalledEngineBuild.xml
@@ -507,7 +507,7 @@
 			<Tag Files="#UAT Win64;#UnrealHeaderTool Win64;#Build Tools Win64" Filter="$(Win64StripFilter)" Except="$(Win64StripExceptions);$(CopyInstalledExceptions)" With="#Strip Tools for Installed Win64"/>
 
 			<!-- Tag files that need to be signed -->
-			<Tag Files="#UAT Win64;#UnrealHeaderTool Win64;#Build Tools Win64" Filter="$(WindowsSignFilter)" Except="$(CopyInstalledExceptions)" With="#Sign Tools for Installed Win64"/>
+			<Tag Files="#UAT Win64;#UnrealHeaderTool Win64;#Build Tools Win64" Filter="$(WindowsSignFilter)" Except="$(CopyInstalledExceptions);.../ThirdParty/..." With="#Sign Tools for Installed Win64"/>
 
 			<!-- Tag files that need to be copied directly -->
 			<Tag Files="#Installed Win64" Except="$(CopyInstalledExceptions);#Strip Tools for Installed Win64;#Sign Tools for Installed Win64" With="#Copy for Installed Win64"/>
@@ -692,7 +692,7 @@
 			<Tag Files="#UnrealHeaderTool Mac;#Build Tools Mac" Filter="$(MacStripFilter)" Except="$(CopyInstalledExceptions)" With="#Strip Tools for Installed Mac"/>
 
 			<!-- Tag files that need to be signed -->
-			<Tag Files="#UnrealHeaderTool Mac;#Build Tools Mac" Filter="$(MacSignFilter)" Except="$(CopyInstalledExceptions)" With="#Sign Tools for Installed Mac"/>
+			<Tag Files="#UnrealHeaderTool Mac;#Build Tools Mac" Filter="$(MacSignFilter)" Except="$(CopyInstalledExceptions);.../ThirdParty/..." With="#Sign Tools for Installed Mac"/>
 
 			<!-- Tag files that need to be copied directly -->
 			<Tag Files="#Installed Mac" Except="$(CopyInstalledExceptions);#Strip Tools for Installed Mac;#Sign Tools for Installed Mac" With="#Copy for Installed Mac"/>
diff --git a/Engine/Source/Programs/AutomationTool/BuildGraph/Tasks/CsCompileTask.cs b/Engine/Source/Programs/AutomationTool/BuildGraph/Tasks/CsCompileTask.cs
--- a/Engine/Source/Programs/AutomationTool/BuildGraph/Tasks/CsCompileTask.cs
+++ b/Engine/Source/Programs/AutomationTool/BuildGraph/Tasks/CsCompileTask.cs
@@ -250,10 +250,10 @@ namespace AutomationTool.Tasks
 						FileReference OutputFile = FileReference.Combine(OutputDir, OtherAssembly.GetFileName());
 						BuildProducts.Add(OutputFile);
 
-						FileReference SymbolFile = OtherAssembly.ChangeExtension(".pdb");
-						if(SymbolFile.Exists())
+						FileReference OutputSymbolFile = OutputFile.ChangeExtension(".pdb");
+						if(OutputSymbolFile.Exists())
 						{
-							BuildProducts.Add(OutputFile.ChangeExtension(".pdb"));
+							BuildProducts.Add(OutputSymbolFile);
 						}
 					}
 					else
