Fix pattern in JunkManifest.txt for online modules that were moved.

diff --git a/Engine/Build/JunkManifest.txt b/Engine/Build/JunkManifest.txt
--- a/Engine/Build/JunkManifest.txt
+++ b/Engine/Build/JunkManifest.txt
@@ -29,12 +29,18 @@ Engine\Binaries\Win64\*-CameraEditor*.*
 Engine\Binaries\Win64\*-CameraDebugger*.*
 Engine\Binaries\Win64\*-SessionLauncher*.*
 Engine\Binaries\Win64\*-OnlineSubsystem*.*
-Engine\Binaries\Win64\*-PurchaseFlow.*
-Engine\Binaries\Win64\*-Hotfix.*
-Engine\Binaries\Win64\*-Lobby.*
-Engine\Binaries\Win64\*-Party.*
-Engine\Binaries\Win64\*-Qos.*
-Engine\Binaries\Win64\*-OnlineBlueprintSupport.*
+Engine\Binaries\Win64\*-PurchaseFlow*.*
+Engine\Binaries\Win32\*-PurchaseFlow*.*
+Engine\Binaries\Win64\*-Hotfix*.*
+Engine\Binaries\Win32\*-Hotfix*.*
+Engine\Binaries\Win64\*-Lobby*.*
+Engine\Binaries\Win32\*-Lobby*.*
+Engine\Binaries\Win64\*-Party*.*
+Engine\Binaries\Win32\*-Party*.*
+Engine\Binaries\Win64\*-Qos*.*
+Engine\Binaries\Win32\*-Qos*.*
+Engine\Binaries\Win64\*-OnlineBlueprintSupport*.*
+Engine\Binaries\Win32\*-OnlineBlueprintSupport*.*
 Engine\Binaries\Mac\UnrealFrontend.app\Contents\Frameworks\*-CoreSlateStyle.*
 Engine\Binaries\Mac\UnrealFrontend.app\Contents\Frameworks\*-SlateStyle.*
 Engine\Binaries\Mac\UnrealFrontend.app\Contents\MacOS\*-XCodeSourceCodeAccess.*
@@ -96,6 +102,3 @@ Ocean\Binaries\Win64\*-LootTables*
 Ocean\Binaries\Win32\*-LootTables*
 Ocean\Binaries\Mac\Ocean.app\Contents\MacOS\*-LootTables.dylib
 Engine\Intermediate\Build\Win64\UE4Editor\Inc\SimplygonUtilities\*
-Engine\Binaries\Win64\*-OnlineBlueprintSupport*.*
-Engine\Binaries\Win32\*-OnlineBlueprintSupport*.*
-
