Fix for pattern in JunkManifest.txt deleting QosReporter when it was only intended to delete Qos

diff --git a/Engine/Build/JunkManifest.txt b/Engine/Build/JunkManifest.txt
--- a/Engine/Build/JunkManifest.txt
+++ b/Engine/Build/JunkManifest.txt
@@ -37,8 +37,10 @@ Engine\Binaries\Win64\*-Lobby*.*
 Engine\Binaries\Win32\*-Lobby*.*
 Engine\Binaries\Win64\*-Party*.*
 Engine\Binaries\Win32\*-Party*.*
-Engine\Binaries\Win64\*-Qos*.*
-Engine\Binaries\Win32\*-Qos*.*
+Engine\Binaries\Win64\*-Qos-*.*
+Engine\Binaries\Win32\*-Qos-*.*
+Engine\Binaries\Win64\*-Qos.*
+Engine\Binaries\Win32\*-Qos.*
 Engine\Binaries\Win64\*-OnlineBlueprintSupport*.*
 Engine\Binaries\Win32\*-OnlineBlueprintSupport*.*
 Engine\Binaries\Mac\UnrealFrontend.app\Contents\Frameworks\*-CoreSlateStyle.*
