Fix failure to copy private reference to UnrealBuildTool build products files from NotForLicensees folder on Mac. The files happen to not exist, since the CopyUAT command used by the build system only copies DLLs, but the .../NotForLicensees/... filter was filtering them out on Windows anyway. Since the filter was broken on Mac, it was failing to copy them instead.

diff --git a/Engine/Build/InstalledEngineBuild.xml b/Engine/Build/InstalledEngineBuild.xml
--- a/Engine/Build/InstalledEngineBuild.xml
+++ b/Engine/Build/InstalledEngineBuild.xml
@@ -611,7 +611,7 @@
 				<Tag Files="#UE4Game Mac;#UE4Game Mac Includes" Except="#UE4Game Mac Unstripped;#UE4Game Mac Unsigned" With="#Installed Mac"/>
 				<Tag Files="#UE4Game Mac Stripped;#UE4Game Mac Signed" With="#Saved Output"/>
 				<Property Name="CopyInstalledFilter" Value="$(CopyInstalledFilter);$(CopyMacFilter)"/>
-				<Property Name="CopyInstalledExceptions" Value="$(CopyInstalledFilter);$(CopyMacExceptions)"/>
+				<Property Name="CopyInstalledExceptions" Value="$(CopyInstalledExceptions);$(CopyMacExceptions)"/>
 			</Do>
 			<Do  If="'$(WithAndroid)' == true">
 				<Tag Files="#UE4Game Android;#UE4Game Android Includes" Except="#UE4Game Android Unstripped" With="#Installed Mac"/>
