Fix build dependencies not being tagged from target receipts on Mac. Causes crash on startup due to missing CEF binaries.

diff --git a/Engine/Build/InstalledEngineBuild.xml b/Engine/Build/InstalledEngineBuild.xml
--- a/Engine/Build/InstalledEngineBuild.xml
+++ b/Engine/Build/InstalledEngineBuild.xml
@@ -651,7 +651,7 @@
 			</Do>
 
 			<!-- Tag any dependencies from all previous build products -->
-			<Tag Files="#Installed Win64" Filter="*.target" With="#TargetReceipts"/>
+			<Tag Files="#Installed Mac" Filter="*.target" With="#TargetReceipts"/>
 			<TagReceipt Files="#TargetReceipts" RuntimeDependencies="true" PrecompiledBuildDependencies="true" PrecompiledRuntimeDependencies="true" With="#Installed Mac"/>
 
 			<!-- Mac-specific dependencies -->
