Added BaseDir to command tagging Saved Output files so that the relative exception filters still work
Moved ConfidentialExceptions to be the last thing that's added to the installed filter so that we should never accidentally overwrite it

diff --git a/Engine/Build/InstalledEngineBuild.xml b/Engine/Build/InstalledEngineBuild.xml
--- a/Engine/Build/InstalledEngineBuild.xml
+++ b/Engine/Build/InstalledEngineBuild.xml
@@ -369,7 +369,7 @@
 
 			<!-- Define filter and exception properties for the installed build -->
 			<Property Name="CopyInstalledFilter" Value="$(CopyEditorFilter)"/>
-			<Property Name="CopyInstalledExceptions" Value="$(ConfidentialExceptions);$(CopyEditorExceptions)"/>
+			<Property Name="CopyInstalledExceptions" Value="$(CopyEditorExceptions)"/>
 
 			<!-- Tag required files from UAT and UBT that will already have been built-->
 			<Tag Files="Engine/Source/Programs/AutomationTool/..." Filter="*.csproj" With="#UAT Projects"/>
@@ -495,6 +495,9 @@
 			<Property Name="CopyInstalledFilter" Value="$(CopyInstalledFilter);Engine/Extras/UnrealVS/..."/>
 			<Property Name="CopyInstalledFilter" Value="$(CopyInstalledFilter);Engine/Extras/Redist/en-us/*"/>
 			
+			<!-- Confidential Exceptions - done last to ensure they never get overwritten -->
+			<Property Name="CopyInstalledExceptions" Value="$(CopyInstalledExceptions);$(ConfidentialExceptions)"/>
+			
 			<!-- Tag additional files needed for installed build -->
 			<Tag Files="$(CopyInstalledFilter)" With="#Installed Win64"/>
 
@@ -509,7 +512,7 @@
 			<Copy Files="#Copy for Installed Win64" FromDir="$(RootDir)" ToDir="$(LocalInstalledDir)" Tag="#Installed Build Win64 Files"/>
 
 			<!-- Tag saved output that needs to be copied directly -->
-			<Tag Files="#Saved Output" Except="$(CopyInstalledExceptions)" With="#Copy Saved Output"/>
+			<Tag BaseDir="$(SavedOutput)" Files="#Saved Output" Except="$(CopyInstalledExceptions)" With="#Copy Saved Output"/>
 			<Copy Files="#Copy Saved Output" FromDir="$(SavedOutput)" ToDir="$(LocalInstalledDir)" Tag="#Installed Build Win64 Files"/>
 
 			<!-- Strip required files into their final location -->
@@ -590,7 +593,7 @@
 
 			<!-- Define filter and exception properties for the installed build -->
 			<Property Name="CopyInstalledFilter" Value="$(CopyEditorFilter)"/>
-			<Property Name="CopyInstalledExceptions" Value="$(ConfidentialExceptions);$(CopyEditorExceptions)"/>
+			<Property Name="CopyInstalledExceptions" Value="$(CopyEditorExceptions)"/>
 
 			<!-- Tag required files from UAT and UBT that will already have been built-->
 			<Tag Files="Engine/Source/Programs/AutomationTool/..." Filter="*.csproj" With="#UAT Projects"/>
@@ -678,6 +681,9 @@
 
 			<!-- Don't want these folders, even if they're part of Windows tools -->
 			<Property Name="CopyInstalledExceptions" Value="$(CopyInstalledExceptions);Engine/Binaries/Win64/...;Engine/Binaries/Win32/..."/>
+
+			<!-- Confidential Exceptions - done last to ensure they never get overwritten -->
+			<Property Name="CopyInstalledExceptions" Value="$(CopyInstalledExceptions);$(ConfidentialExceptions)"/>
 			
 			<!-- Tag additional files needed for editor and platforms -->
 			<Tag Files="$(CopyInstalledFilter)" With="#Installed Mac"/>
@@ -693,7 +699,7 @@
 			<Copy Files="#Copy for Installed Mac" FromDir="$(RootDir)" ToDir="$(LocalInstalledDirMac)" Tag="#Installed Build Mac Files"/>
 
 			<!-- Tag saved output that needs to be copied directly -->
-			<Tag Files="#Saved Output" Except="$(CopyInstalledExceptions)" With="#Copy Saved Output"/>
+			<Tag BaseDir="$(SavedOutput)" Files="#Saved Output" Except="$(CopyInstalledExceptions)" With="#Copy Saved Output"/>
 			<Copy Files="#Copy Saved Output" FromDir="$(SavedOutput)" ToDir="$(LocalInstalledDirMac)" Tag="#Installed Build Mac Files"/>
 
 			<!-- Strip required files into their final location -->
