Add missing SwarmInterface compile to tools group.

diff --git a/Engine/Build/InstalledEngineBuild.xml b/Engine/Build/InstalledEngineBuild.xml
--- a/Engine/Build/InstalledEngineBuild.xml
+++ b/Engine/Build/InstalledEngineBuild.xml
@@ -224,6 +224,7 @@
 			<Compile Target="BootstrapPackagedGame" Configuration="Shipping" Platform="Win64" Tag="#Build Tools Win64"/>
 			<Compile Target="BuildPatchTool" Configuration="Shipping" Platform="Win64" Tag="#Build Tools Win64"/>
 
+			<CsCompile Project="Engine\Source\Editor\SwarmInterface\DotNET\SwarmInterface.csproj" Configuration="Development" Platform="AnyCPU" Tag="#Build Tools Win64" TagReferences="#Build Tools Win64"/>
 			<CsCompile Project="Engine\Source\Programs\UnrealSwarm\SwarmCoordinator\SwarmCoordinator.csproj" Configuration="Development" Platform="AnyCPU" Tag="#Build Tools Win64" TagReferences="#Build Tools Win64"/>
 			<CsCompile Project="Engine\Source\Programs\UnrealSwarm\Agent\Agent.csproj" Configuration="Development" Platform="AnyCPU" Tag="#Build Tools Win64" TagReferences="#Build Tools Win64"/>
 			<CsCompile Project="Engine\Source\Programs\NetworkProfiler\NetworkProfiler\NetworkProfiler.csproj" Configuration="Development" Platform="AnyCPU" Tag="#Build Tools Win64" TagReferences="#Build Tools Win64"/>
