Add missing mono distribution for Mac. Fixes crash on packaging.

diff --git a/Engine/Build/InstalledEngineBuild.xml b/Engine/Build/InstalledEngineBuild.xml
--- a/Engine/Build/InstalledEngineBuild.xml
+++ b/Engine/Build/InstalledEngineBuild.xml
@@ -677,6 +677,8 @@
 			<Property Name="CopyInstalledFilter" Value="$(CopyInstalledFilter);Engine/Build/BatchFiles/RunUAT.command"/>
 			<Property Name="CopyInstalledFilter" Value="$(CopyInstalledFilter);Engine/Build/BatchFiles/RunUAT.sh"/>
 
+			<Property Name="CopyInstalledFilter" Value="$(CopyInstalledFilter);Engine/Binaries/ThirdParty/Mono/Mac/..."/>
+
 			<Property Name="CopyInstalledFilter" Value="$(CopyInstalledFilter);Engine/Extras/LLDBDataFormatters/UE4DataFormatters.py"/>
 			<Property Name="CopyInstalledFilter" Value="$(CopyInstalledFilter);Engine/Extras/Maya_AnimationRiggingTools/ArtToolsOSX/..."/>
 
