Add missing binaries required to package on Win64.

diff --git a/Engine/Build/InstalledEngineBuild.xml b/Engine/Build/InstalledEngineBuild.xml
--- a/Engine/Build/InstalledEngineBuild.xml
+++ b/Engine/Build/InstalledEngineBuild.xml
@@ -488,6 +488,9 @@
 			<Property Name="CopyInstalledFilter" Value="$(CopyInstalledFilter);Engine/Build/BatchFiles/RunUAT.bat"/>
 			<Property Name="CopyInstalledFilter" Value="$(CopyInstalledFilter);Engine/Build/BatchFiles/GetVSComnToolsPath.bat"/>
 
+			<!-- Binaries which are explicitly staged by UAT rather than listing runtime dependencies -->
+			<Property Name="CopyInstalledFilter" Value="$(CopyInstalledFilter);Engine/Binaries/ThirdParty/OpenSSL/..."/>
+		
 			<!-- Extra tools -->
 			<Property Name="CopyInstalledFilter" Value="$(CopyInstalledFilter);Engine/Extras/3dsMaxScripts/..."/>
 			<Property Name="CopyInstalledFilter" Value="$(CopyInstalledFilter);Engine/Extras/VisualStudioDebugging/UE4.natvis"/>
