Manually include Engine/Binaries/ThirdParty/OpenAL for Linux. OpenAL.build.cs already references this in its final location, but UAT still has code to manually stage it from its original location. Just include both for the time being, until Dmitry can verify out UEPLAT-807.

diff --git a/Engine/Build/InstalledEngineFilters.xml b/Engine/Build/InstalledEngineFilters.xml
--- a/Engine/Build/InstalledEngineFilters.xml
+++ b/Engine/Build/InstalledEngineFilters.xml
@@ -184,6 +184,7 @@
 	<Property Name="CopyLinuxExceptions" Value="$(CopyLinuxExceptions);Engine/Plugins/.../Binaries/Linux/....a"/>
 	<Property Name="CopyLinuxExceptions" Value="$(CopyLinuxExceptions);Engine/Source/ThirdParty/.../Linux/....a"/>
 	<Property Name="CopyLinuxFilter" Value="$(CopyLinuxFilter);Engine/Binaries/Linux/CrashReportClient"/>
+	<Property Name="CopyLinuxFilter" Value="$(CopyLinuxFilter);Engine/Binaries/ThirdParty/OpenAL/..."/>
 	<Property Name="CopyLinuxFilter" Value="$(CopyLinuxFilter);Engine/Binaries/ThirdParty/ICU/icu4c-53_1/Linux/x86_64-unknown-linux-gnu/*.so"/>
 
 	<!-- HTML5 -->
