Duplicating CL#3076974 from Release-4.13

Added Node to build CrashReportClient for Linux that can be used by internal game targets

diff --git a/Engine/Build/InstalledEngineBuild.xml b/Engine/Build/InstalledEngineBuild.xml
--- a/Engine/Build/InstalledEngineBuild.xml
+++ b/Engine/Build/InstalledEngineBuild.xml
@@ -245,6 +245,9 @@
 				<PakFile Output="FeaturePacks/$(SampleName).upack" Files="" ResponseFile="Samples/$(SampleName)/contents.txt"/>
 			</ForEach>
 		</Node>
+		<Node Name="Build Tools Linux" Requires="Compile UnrealHeaderTool Win64">
+			<Compile Target="CrashReportClient" Platform="Linux" Configuration="Shipping"/>
+		</Node>
 	</Agent>
 
 	<Agent Name="Tools Group Mac" Type="CompileMac">
