Disable HTML5 in binary releases until we have a fix for the ICE after the Dev-Physics merge.

diff --git a/Engine/Build/InstalledEngineBuild.xml b/Engine/Build/InstalledEngineBuild.xml
--- a/Engine/Build/InstalledEngineBuild.xml
+++ b/Engine/Build/InstalledEngineBuild.xml
@@ -12,7 +12,7 @@
 	<Option Name="WithIOS" Restrict="true|false" DefaultValue="true" Description="Include the iOS target platform"/>
 	<Option Name="WithTVOS" Restrict="true|false" DefaultValue="true" Description="Include the tvOS target platform"/>
 	<Option Name="WithLinux" Restrict="true|false" DefaultValue="true" Description="Include the Linux target platform"/>
-	<Option Name="WithHTML5" Restrict="true|false" DefaultValue="true" Description="Include the HTML5 target platform"/>
+	<Option Name="WithHTML5" Restrict="true|false" DefaultValue="false" Description="Include the HTML5 target platform"/>
 	<Option Name="WithPS4" Restrict="true|false" DefaultValue="false" Description="Include the PS4 target platform"/>
 	<Option Name="WithXboxOne" Restrict="true|false" DefaultValue="false" Description="Include the XboxOne target platform"/>
 
