Add NotForLicensees SimplygonSwarm DLLs folders to JunkManifest; are now distributed in the clear.

diff --git a/Engine/Build/JunkManifest.txt b/Engine/Build/JunkManifest.txt
--- a/Engine/Build/JunkManifest.txt
+++ b/Engine/Build/JunkManifest.txt
@@ -58,6 +58,7 @@ Engine\Binaries\Mac\*-MessagingDebugger.*
 Engine\Binaries\Win64\*-DDCUtils*.*
 Engine\Binaries\Win64\*-MessagingDebugger.*
 Engine\Binaries\Win64\NotForLicensees\*-SimplygonMeshReduction*.*
+Engine\Binaries\Win64\NotForLicensees\*-SimplygonSwarm*.*
 Engine\Binaries\Win64\*-TextureXboxOneFormat.*
 Engine\Binaries\Win64\*-TextureFormatATITC.*
 Shadow\Binaries\Win64\*ShadowLoadingScreen.*
