Merging CL 3185489 from Dev-Build: Add "Switch" as a confidential platform.

diff --git a/Engine/Build/InstalledEngineBuild.xml b/Engine/Build/InstalledEngineBuild.xml
--- a/Engine/Build/InstalledEngineBuild.xml
+++ b/Engine/Build/InstalledEngineBuild.xml
@@ -15,6 +15,7 @@
 	<Option Name="WithHTML5" Restrict="true|false" DefaultValue="true" Description="Include the HTML5 target platform"/>
 	<Option Name="WithPS4" Restrict="true|false" DefaultValue="false" Description="Include the PS4 target platform"/>
 	<Option Name="WithXboxOne" Restrict="true|false" DefaultValue="false" Description="Include the XboxOne target platform"/>
+	<Option Name="WithSwitch" Restrict="true|false" DefaultValue="false" Description="Include the Switch target platform"/>
 
 	<!-- Whether to create a prebuilt DDC -->
 	<Option Name="WithDDC" Restrict="true|false" DefaultValue="true" Description="Build a standalone derived-data cache for the engine content and templates" />
diff --git a/Engine/Build/InstalledEngineFilters.xml b/Engine/Build/InstalledEngineFilters.xml
--- a/Engine/Build/InstalledEngineFilters.xml
+++ b/Engine/Build/InstalledEngineFilters.xml
@@ -5,6 +5,7 @@
 	<Property Name="ConfidentialExceptions" Value=".../EpicInternal/...;.../CarefullyRedist/...;.../NotForLicensees/...;.../NoRedist/..."/>
 	<Property Name="ConfidentialExceptions" Value="$(ConfidentialExceptions);.../PS4/..." If="'$(WithPS4)' != true"/>
 	<Property Name="ConfidentialExceptions" Value="$(ConfidentialExceptions);.../XboxOne/..." If="'$(WithXboxOne)' != true"/>
+	<Property Name="ConfidentialExceptions" Value="$(ConfidentialExceptions);.../Switch/..." If="'$(WithSwitch)' != true"/>
 
 	<!-- List of file types to be stripped and signed for different platforms -->
 	<Property Name="Win64StripFilter" Value="*.pdb"/>
diff --git a/Engine/Source/Programs/AutomationTool/AutomationUtils/FileFilter.cs b/Engine/Source/Programs/AutomationTool/AutomationUtils/FileFilter.cs
--- a/Engine/Source/Programs/AutomationTool/AutomationUtils/FileFilter.cs
+++ b/Engine/Source/Programs/AutomationTool/AutomationUtils/FileFilter.cs
@@ -358,6 +358,7 @@ namespace AutomationTool
 		{
 			AddRule(".../PS4/...", FileFilterType.Exclude);
 			AddRule(".../XboxOne/...", FileFilterType.Exclude);
+			AddRule(".../Switch/...", FileFilterType.Exclude);
 		}
 
 		/// <summary>
