Update BuildGraph Schema to use Switch instead of WolfPlat and include Wolf directories in installed build exclusions for safety.

diff --git a/Engine/Build/InstalledEngineFilters.xml b/Engine/Build/InstalledEngineFilters.xml
--- a/Engine/Build/InstalledEngineFilters.xml
+++ b/Engine/Build/InstalledEngineFilters.xml
@@ -5,7 +5,7 @@
 	<Property Name="ConfidentialExceptions" Value=".../EpicInternal/...;.../CarefullyRedist/...;.../NotForLicensees/...;.../NoRedist/..."/>
 	<Property Name="ConfidentialExceptions" Value="$(ConfidentialExceptions);.../PS4/..." If="'$(WithPS4)' != true"/>
 	<Property Name="ConfidentialExceptions" Value="$(ConfidentialExceptions);.../XboxOne/..." If="'$(WithXboxOne)' != true"/>
-	<Property Name="ConfidentialExceptions" Value="$(ConfidentialExceptions);.../Switch/..." If="'$(WithSwitch)' != true"/>
+	<Property Name="ConfidentialExceptions" Value="$(ConfidentialExceptions);.../Switch/...;.../Wolf/...;.../WolfPlat/..." If="'$(WithSwitch)' != true"/>
 
 	<!-- List of file types to be stripped and signed for different platforms -->
 	<Property Name="Win64StripFilter" Value="*.pdb"/>
