Another missing .target.cs template file. (PR #3198)

diff --git a/Engine/Build/Target.cs.template b/Engine/Build/Target.cs.template
--- a/Engine/Build/Target.cs.template
+++ b/Engine/Build/Target.cs.template
@@ -5,7 +5,7 @@ using System.Collections.Generic;
 
 public class {GAME_NAME}Target : TargetRules
 {
-    public {GAME_NAME}Target(TargetInfo Target)
+    public {GAME_NAME}Target(TargetInfo Target) : base(Target)
 	{
 		Type = TargetType.Game;
 
