Include plugin config files when staging the installed engine build. Fixes warnings building DDC in CIS.

diff --git a/Engine/Build/InstalledEngineFilters.xml b/Engine/Build/InstalledEngineFilters.xml
--- a/Engine/Build/InstalledEngineFilters.xml
+++ b/Engine/Build/InstalledEngineFilters.xml
@@ -82,6 +82,7 @@
 	<!-- Plugins -->
 	<Property Name="CopyEditorFilter" Value="$(CopyEditorFilter);Engine/Plugins/....uplugin"/>
 	<Property Name="CopyEditorFilter" Value="$(CopyEditorFilter);Engine/Plugins/.../Content/..."/>
+	<Property Name="CopyEditorFilter" Value="$(CopyEditorFilter);Engine/Plugins/.../Config/..."/>
 	<Property Name="CopyEditorFilter" Value="$(CopyEditorFilter);Engine/Plugins/.../Resources/..."/>
 	<Property Name="CopyEditorFilter" Value="$(CopyEditorFilter);Engine/Plugins/.../Source/..."/>
 	<Property Name="CopyEditorFilter" Value="$(CopyEditorFilter);Engine/Plugins/.../Templates/..."/>
