Include config files when staging the mini-engine distribution to build DDC. Should fix CIS warnings.

diff --git a/Engine/Build/InstalledEngineBuild.xml b/Engine/Build/InstalledEngineBuild.xml
--- a/Engine/Build/InstalledEngineBuild.xml
+++ b/Engine/Build/InstalledEngineBuild.xml
@@ -402,7 +402,7 @@
 			<TagReceipt Files="#TargetReceipts" RuntimeDependencies="true" With="#ToCopy"/>
 			<Tag Files="Engine\Content\..." Except="*.psd;*.pdn;*.fbx;*.po" With="#ToCopy"/>
 			<Tag Files="Engine\Config\..." Except="*.vdf" With="#ToCopy"/>
-			<Tag Files="Engine\Plugins\....uplugin;Engine\Plugins\...\Content\...;Engine\Plugins\...\Resources\...;Engine\Plugins\...\Templates\..." Except="Engine\Plugins\Runtime\TwitchLiveStreaming\...;Engine\Plugins\...\Binaries/Mac\..." With="#ToCopy"/>
+			<Tag Files="Engine\Plugins\....uplugin;Engine\Plugins\...\Config\...;Engine\Plugins\...\Content\...;Engine\Plugins\...\Resources\...;Engine\Plugins\...\Templates\..." Except="Engine\Plugins\Runtime\TwitchLiveStreaming\...;Engine\Plugins\...\Binaries/Mac\..." With="#ToCopy"/>
 			<Tag Files="Engine\Shaders\..." With="#ToCopy"/>
 			<Tag Files="Engine\Binaries\DotNET\Ionic.Zip.Reduced.dll" With="#ToCopy"/>
 			<Tag Files="Engine\Binaries\DotNET\OneSky.dll" With="#ToCopy"/>
