Another fixup for warning building DDC. Include plugin config files on Mac.

diff --git a/Engine/Build/InstalledEngineBuild.xml b/Engine/Build/InstalledEngineBuild.xml
--- a/Engine/Build/InstalledEngineBuild.xml
+++ b/Engine/Build/InstalledEngineBuild.xml
@@ -442,7 +442,7 @@
 			<TagReceipt Files="#TargetReceipts" RuntimeDependencies="true" With="#ToCopy"/>
 			<Tag Files="Engine/Content/..." Except="*.psd;*.pdn;*.fbx;*.po" With="#ToCopy"/>
 			<Tag Files="Engine/Config/..." Except="*.vdf" With="#ToCopy"/>
-			<Tag Files="Engine/Plugins/....uplugin;Engine/Plugins/.../Content/...;Engine/Plugins/.../Resources/...;Engine/Plugins/.../Templates/..." Except="Engine/Plugins/Runtime/TwitchLiveStreaming/...;Engine/Plugins/.../Binaries/Win64/...;Engine/Plugins/.../Binaries/Win32/..." With="#ToCopy"/>
+			<Tag Files="Engine/Plugins/....uplugin;Engine/Plugins/.../Config/...;Engine/Plugins/.../Content/...;Engine/Plugins/.../Resources/...;Engine/Plugins/.../Templates/..." Except="Engine/Plugins/Runtime/TwitchLiveStreaming/...;Engine/Plugins/.../Binaries/Win64/...;Engine/Plugins/.../Binaries/Win32/..." With="#ToCopy"/>
 			<Tag Files="Engine/Shaders/..." With="#ToCopy"/>
 			<Tag Files="Engine/Binaries/DotNET/Ionic.Zip.Reduced.dll" With="#ToCopy"/>
 			<Tag Files="Engine/Binaries/DotNET/OneSky.dll" With="#ToCopy"/>
