Fix UBT and UAT not being included in installed builds on Mac and Linux, causing compilation to fail, and "Launch on" option not being available.

diff --git a/Engine/Build/InstalledEngineBuild.xml b/Engine/Build/InstalledEngineBuild.xml
--- a/Engine/Build/InstalledEngineBuild.xml
+++ b/Engine/Build/InstalledEngineBuild.xml
@@ -21,8 +21,7 @@
 	<Option Name="WithDDC" Restrict="true|false" DefaultValue="true" Description="Build a standalone derived-data cache for the engine content and templates" />
 
 	<!-- Whether to build DDC for the host platform only -->
-	
-<Option Name="HostPlatformDDCOnly" Restrict="true|false" DefaultValue="true" Description="Whether to include DDC for the host platform only"/>
+	<Option Name="HostPlatformDDCOnly" Restrict="true|false" DefaultValue="true" Description="Whether to include DDC for the host platform only"/>
 
 	<!-- Whether to sign any executables produced -->
 	<Option Name="SignExecutables" Restrict="true|false" DefaultValue="false" Description="Sign the executables produced where signing is available"/>
@@ -773,7 +772,7 @@
 			</Do>
 
 			<!-- Copy the stripped version of UAT -->
-			<Copy Files="#StrippedUAT" From="$(StrippedUATDir)" To="$(LocalInstalledDir)" Tag="#Installed Build Mac Files"/>
+			<Copy Files="#StrippedUAT" From="$(StrippedUATDir)" To="$(LocalInstalledDirMac)" Tag="#Installed Build Mac Files"/>
 
 			<!-- Copy the compressed DDC -->
 			<Property Name="SavedDir" Value="$(SavedOutput)/Installed/Mac"/>
@@ -846,7 +845,7 @@
 			<Strip Files="#Strip Tools for Installed Linux" BaseDir="$(RootDir)" OutputDir="$(LocalInstalledDirLinux)" Platform="Linux" Tag="#Installed Build Linux Files"/>
 
 			<!-- Copy the stripped version of UAT -->
-			<Copy Files="#StrippedUAT" From="$(StrippedUATDir)" To="$(LocalInstalledDir)" Tag="#Installed Build Linux Files"/>
+			<Copy Files="#StrippedUAT" From="$(StrippedUATDir)" To="$(LocalInstalledDirLinux)" Tag="#Installed Build Linux Files"/>
 
 			<!-- Copy the compressed DDC -->
 			<Property Name="SavedDir" Value="$(SavedOutput)\Installed\Linux"/>
