Build: Exclude UHT timestamp files from the build products included in the binary release. These files are always modified, which causes downstream build steps which regenerate them to fail.

diff --git a/Engine/Build/InstalledEngineBuild.xml b/Engine/Build/InstalledEngineBuild.xml
--- a/Engine/Build/InstalledEngineBuild.xml
+++ b/Engine/Build/InstalledEngineBuild.xml
@@ -113,7 +113,7 @@
 			</Do>
 
 			<!-- Tag the generated includes for this target -->
-			<Tag Files="Engine/Intermediate/Build/Win64/UE4Editor/Inc/...;Engine/Plugins/.../Intermediate/Build/Win64/UE4Editor/Inc/..." With="#UE4Editor Win64 Includes"/>
+			<Tag Files="Engine/Intermediate/Build/Win64/UE4Editor/Inc/...;Engine/Plugins/.../Intermediate/Build/Win64/UE4Editor/Inc/..." Except=".../Timestamp" With="#UE4Editor Win64 Includes"/>
 		</Node>
 	</Agent>
 	
