Build: Exclude the UHT timestamp files from the list of build products for Mac and Linux to fix build failures.

diff --git a/Engine/Build/InstalledEngineBuild.xml b/Engine/Build/InstalledEngineBuild.xml
--- a/Engine/Build/InstalledEngineBuild.xml
+++ b/Engine/Build/InstalledEngineBuild.xml
@@ -136,7 +136,7 @@
 			</Do>
 		
 			<!-- Tag the generated includes for this target -->
-			<Tag Files="Engine/Intermediate/Build/Mac/UE4Editor/Inc/...;Engine/Plugins/.../Intermediate/Build/Mac/UE4Editor/Inc/..." With="#UE4Editor Mac Includes"/>
+			<Tag Files="Engine/Intermediate/Build/Mac/UE4Editor/Inc/...;Engine/Plugins/.../Intermediate/Build/Mac/UE4Editor/Inc/..." Except=".../Timestamp" With="#UE4Editor Mac Includes"/>
 		</Node>
 	</Agent>
 
@@ -152,7 +152,7 @@
 			<Strip BaseDir="$(RootDir)" OutputDir="$(SavedOutput)" Platform="Linux" Files="#UE4Editor Linux Unstripped" Tag="#UE4Editor Linux Stripped"/>
 
 			<!-- Tag the generated includes for this target -->
-			<Tag Files="Engine/Intermediate/Build/Linux/B4D820EA/UE4Editor/Inc/...;Engine/Plugins/.../Intermediate/Build/Linux/B4D820EA/UE4Editor/Inc/..." With="#UE4Editor Linux Includes"/>
+			<Tag Files="Engine/Intermediate/Build/Linux/B4D820EA/UE4Editor/Inc/...;Engine/Plugins/.../Intermediate/Build/Linux/B4D820EA/UE4Editor/Inc/..." Except=".../Timestamp" With="#UE4Editor Linux Includes"/>
 		</Node>
 	</Agent>
 
