Fix supported agent types for compiling PhysX on Mac.

diff --git a/Engine/Build/PhysXBuild.xml b/Engine/Build/PhysXBuild.xml
--- a/Engine/Build/PhysXBuild.xml
+++ b/Engine/Build/PhysXBuild.xml
@@ -99,17 +99,17 @@
 			<Command Name="BuildPhysX" Arguments="-TargetPlatforms=HTML5 -TargetConfigs=&quot;$(TargetConfigs)&quot; -TargetWindowsCompilers=&quot;$(TargetWindowsCompilers)&quot; $(OtherArgs)"/>
 		</Node>
 	</Agent>
-	<Agent Name="PhysX Mac" Type="EngineMac;PhysX_CompileMac">
+	<Agent Name="PhysX Mac" Type="EngineMac;PhysX_CompileMac;Mac">
 		<Node Name="Compile PhysX Mac">
 			<Command Name="BuildPhysX" Arguments="-TargetPlatforms=Mac -TargetConfigs=&quot;$(TargetConfigs)&quot; $(OtherArgs)"/>
 		</Node>
 	</Agent>
-	<Agent Name="PhysX IOS" Type="EngineMac;PhysX_CompileMac">
+	<Agent Name="PhysX IOS" Type="EngineMac;PhysX_CompileMac;Mac">
 		<Node Name="Compile PhysX IOS">
 			<Command Name="BuildPhysX" Arguments="-TargetPlatforms=IOS -TargetConfigs=&quot;$(TargetConfigs)&quot; $(OtherArgs)"/>
 		</Node>
 	</Agent>
-	<Agent Name="PhysX TVOS" Type="EngineMac;PhysX_CompileMac">
+	<Agent Name="PhysX TVOS" Type="EngineMac;PhysX_CompileMac;Mac">
 		<Node Name="Compile PhysX TVOS">
 			<Command Name="BuildPhysX" Arguments="-TargetPlatforms=TVOS -TargetConfigs=&quot;$(TargetConfigs)&quot; $(OtherArgs)"/>
 		</Node>
