Add support for building Linux code projects from Windows using the cross-compile toolchain in the installed engine build.

diff --git a/Engine/Build/InstalledEngineBuild.xml b/Engine/Build/InstalledEngineBuild.xml
--- a/Engine/Build/InstalledEngineBuild.xml
+++ b/Engine/Build/InstalledEngineBuild.xml
@@ -670,7 +670,7 @@
 			<Copy Files="#CompressedDDCWin64" From="$(SavedDir)" To="$(LocalInstalledDir)" If="'$(WithDDC)' == true" Tag="#Installed Build Win64 Files"/>
 
 			<!-- List all of the platforms that we only allow for content projects -->
-			<Property Name="ContentOnlyPlatforms" Value="Linux"/>
+			<Property Name="ContentOnlyPlatforms" Value=""/>
 
 			<!-- Create command line arguments for any platforms that were built with additional architectures -->
 			<Property Name="ArchitectureCommands" Value=""/>
diff --git a/Engine/Build/InstalledEngineFilters.xml b/Engine/Build/InstalledEngineFilters.xml
--- a/Engine/Build/InstalledEngineFilters.xml
+++ b/Engine/Build/InstalledEngineFilters.xml
@@ -377,15 +377,6 @@
 <Property Name="CopyLinuxFilter" Value="$(CopyLinuxFilter);Engine/Extras/ThirdPartyNotUE/putty/..."/>	
 
 <!-- Win64 specific -->
-	<Do If="'$(HostPlatform)' == 'Win64'">
-		<Property Name="CopyLinuxExceptions" Value="$(CopyLinuxExceptions)">
-			<!-- Content only - exclude things we only need for code projects -->
-			Engine/Binaries/Linux/....a
-			Engine/Plugins/.../Binaries/Linux/....a
-			Engine/Source/ThirdParty/.../Linux/....a
-		</Property>
-	</Do>
-
 	<!-- HTML5 -->
 	<Property Name="CopyHTML5Filter">
 		<!-- Build Files -->
