Removing ARKit from the InstalledEngineBuild.xml, since we disabled it on the build farm temporarily

diff --git a/Engine/Build/InstalledEngineBuild.xml b/Engine/Build/InstalledEngineBuild.xml
--- a/Engine/Build/InstalledEngineBuild.xml
+++ b/Engine/Build/InstalledEngineBuild.xml
@@ -294,10 +294,6 @@
 			<Compile Target="UE4Game" Platform="IOS" Configuration="Development" Tag="#UE4Game IOS" Arguments="-precompile"/>
 			<Compile Target="UE4Game" Platform="IOS" Configuration="Shipping" Tag="#UE4Game IOS" Arguments="-precompile"/>
 
-			<!-- Manually add precompiled ARKit binaries, since they're built with an Xcode beta. -->
-			<Tag Files="$(RootDir)/Engine/Plugins/Experimental/AppleARKit/Binaries/IOS/UE4-AppleARKit.a" With="#UE4Game IOS"/>
-			<Tag Files="$(RootDir)/Engine/Plugins/Experimental/AppleARKit/Binaries/IOS/UE4-AppleARKit-IOS-Shipping.a" With="#UE4Game IOS"/>
-
 			<Tag Files="#UE4Game IOS" Filter="$(IOSStripFilter)" With="#UE4Game IOS Unstripped"/>
 			<Strip BaseDir="$(RootDir)" OutputDir="$(SavedOutput)" Platform="IOS" Files="#UE4Game IOS Unstripped" Tag="#UE4Game IOS Stripped"/>
 
