Re-enabling ARKit for 4.18

diff --git a/Engine/Build/InstalledEngineBuild.xml b/Engine/Build/InstalledEngineBuild.xml
--- a/Engine/Build/InstalledEngineBuild.xml
+++ b/Engine/Build/InstalledEngineBuild.xml
@@ -294,6 +294,10 @@
 			<Compile Target="UE4Game" Platform="IOS" Configuration="Development" Tag="#UE4Game IOS" Arguments="-precompile"/>
 			<Compile Target="UE4Game" Platform="IOS" Configuration="Shipping" Tag="#UE4Game IOS" Arguments="-precompile"/>
 
+			<!-- Manually add precompiled ARKit binaries, since they're built with an Xcode beta. -->
+			<Tag Files="$(RootDir)/Engine/Plugins/Experimental/AppleARKit/Binaries/IOS/UE4-AppleARKit.a" With="#UE4Game IOS"/>
+			<Tag Files="$(RootDir)/Engine/Plugins/Experimental/AppleARKit/Binaries/IOS/UE4-AppleARKit-IOS-Shipping.a" With="#UE4Game IOS"/>
+
 			<Tag Files="#UE4Game IOS" Filter="$(IOSStripFilter)" With="#UE4Game IOS Unstripped"/>
 			<Strip BaseDir="$(RootDir)" OutputDir="$(SavedOutput)" Platform="IOS" Files="#UE4Game IOS Unstripped" Tag="#UE4Game IOS Stripped"/>
 
diff --git a/Engine/Plugins/Runtime/AppleARKit/AppleARKit.uplugin b/Engine/Plugins/Runtime/AppleARKit/AppleARKit.uplugin
--- a/Engine/Plugins/Runtime/AppleARKit/AppleARKit.uplugin
+++ b/Engine/Plugins/Runtime/AppleARKit/AppleARKit.uplugin
@@ -12,18 +12,19 @@
 	"SupportURL": "",
 	"CanContainContent": true,
 	"IsBetaVersion": true,
+    "EnabledByDefault": true,
 	"Modules": [
 		{
 			"Name": "AppleARKit",
 			"Type": "Runtime",
 			"LoadingPhase": "PostConfigInit",
-			"WhitelistPlatforms" : [ "Win64" ]
+			"WhitelistPlatforms" : [ "Mac", "iOS", "Win64" ]
 		},
 		{
 			"Name": "AppleARKitEditor",
 			"Type": "Editor",
 			"LoadingPhase": "Default",
-			"WhitelistPlatforms" : [ "Win64" ]
+			"WhitelistPlatforms" : [ "Mac", "iOS", "Win64" ]
 		}
 	]
 }
