Fix compile errors compiling C# tools for public, due to Visual Studio interop DLL not being copied to staging directory.

diff --git a/Engine/Build/InstalledEngineBuild.xml b/Engine/Build/InstalledEngineBuild.xml
--- a/Engine/Build/InstalledEngineBuild.xml
+++ b/Engine/Build/InstalledEngineBuild.xml
@@ -447,7 +447,7 @@
 			<!-- Copy Source and referenced libraries to a new location with Confidential folders removed-->
 			<Tag Files="Engine/Source/Programs/AutomationTool/..." Filter="*.csproj" Except="$(ConfidentialExceptions)"  With="#UAT Projects"/>
 			<CsCompile Project="#UAT Projects" Configuration="Development" Platform="AnyCPU" TagReferences="#UATReferences" EnumerateOnly="true"/>
-			<Tag Files="Engine/Binaries/DotNET/...;Engine/Source/Programs/...;Engine/Source/Editor/SwarmInterface/...;#UATReferences" Except="$(ConfidentialExceptions);AutomationTool.csproj.References" With="#RedistUATSource"/>
+			<Tag Files="Engine/Binaries/DotNET/...;Engine/Binaries/ThirdParty/VisualStudio/...;Engine/Source/Programs/...;Engine/Source/Editor/SwarmInterface/...;#UATReferences" Except="$(ConfidentialExceptions);AutomationTool.csproj.References" With="#RedistUATSource"/>
 			<Copy Files="#RedistUATSource" From="$(RootDir)" To="$(CsToolsDir)"/>
 			<Tag Files="$(CsToolsDir)/Engine/Source/Programs/AutomationTool/..." Filter="*.csproj" With="#RedistUATProjects"/>
 
