Fix warnings building DDC.

- Check for the exisence of EpicInternal.txt to determine whether to copy DDCUtils.dll to the staging directory for generating DDC. These conditions are evaluated before dependencies are read from the temp storage directory, so we can't expect them to be there.
- Copy DDCUtils for ShaderCompileWorker, since it has its own copy.
- Resave InstallingVisualStudioTutorial.uasset with correct engine verison.

diff --git a/Engine/Build/InstalledEngineBuild.xml b/Engine/Build/InstalledEngineBuild.xml
--- a/Engine/Build/InstalledEngineBuild.xml
+++ b/Engine/Build/InstalledEngineBuild.xml
@@ -538,10 +538,12 @@
 			<Tag Files="Templates\TemplateResources\..." With="#ToCopy"/>
 			<!-- Filter out the files not needed to build DDC. Removing confidential folders can affect DDC keys, so we want to be sure that we're making DDC with a build that can use it. -->
 			<Tag Files="#ToCopy" Except="$(ConfidentialExceptions);.../Source/...;.../Intermediate/..." With="#FilteredCopyList"/>
-			<!-- Additional required files that would be removed by confidential filter -->
+			<!-- Additional required files that would be removed by confidential filter. Note that conditions are evaluated before copying build products, so check against EpicInternal.txt -->
 			<Tag Files="Engine\Build\NotForLicensees\EpicInternal.txt" With="#FilteredCopyList" If="Exists('Engine\Build\NotForLicensees\EpicInternal.txt')"/>
-			<Tag Files="Engine\Binaries\Win64\NotForLicensees\UE4Editor.modules" With="#FilteredCopyList" If="Exists('Engine\Binaries\Win64\NotForLicensees\UE4Editor.modules')"/>
-			<Tag Files="Engine\Binaries\Win64\NotForLicensees\UE4Editor-DDCUtils.dll" With="#FilteredCopyList" If="Exists('Engine\Binaries\Win64\NotForLicensees\UE4Editor-DDCUtils.dll')"/>
+			<Tag Files="Engine\Binaries\Win64\NotForLicensees\ShaderCompileWorker.modules" With="#FilteredCopyList" If="Exists('Engine\Build\NotForLicensees\EpicInternal.txt')"/>
+			<Tag Files="Engine\Binaries\Win64\NotForLicensees\ShaderCompileWorker-DDCUtils.dll" With="#FilteredCopyList" If="Exists('Engine\Build\NotForLicensees\EpicInternal.txt')"/>
+			<Tag Files="Engine\Binaries\Win64\NotForLicensees\UE4Editor.modules" With="#FilteredCopyList" If="Exists('Engine\Build\NotForLicensees\EpicInternal.txt')"/>
+			<Tag Files="Engine\Binaries\Win64\NotForLicensees\UE4Editor-DDCUtils.dll" With="#FilteredCopyList" If="Exists('Engine\Build\NotForLicensees\EpicInternal.txt')"/>
 			<!-- Copy everything to a temporary directory -->
 			<Copy From="$(RootDir)" To="$(DDCDir)" Files="#FilteredCopyList"/>
 
@@ -582,8 +584,10 @@
 			<Tag Files="#ToCopy" Except="$(ConfidentialExceptions);.../Source/...;.../Intermediate/..." With="#FilteredCopyList"/>
 			<!-- Additional required files that would be removed by confidential filter -->
 			<Tag Files="Engine/Build/NotForLicensees/EpicInternal.txt" With="#FilteredCopyList" If="Exists('Engine/Build/NotForLicensees/EpicInternal.txt')"/>
-			<Tag Files="Engine/Binaries/Mac/UE4Editor.app/Contents/MacOS/NotForLicensees/UE4Editor.modules" With="#FilteredCopyList" If="Exists('Engine/Binaries/Mac/UE4Editor.app/Contents/MacOS/NotForLicensees/UE4Editor.modules')"/>
-			<Tag Files="Engine/Binaries/Mac/UE4Editor.app/Contents/MacOS/NotForLicensees/UE4Editor-DDCUtils.dylib" With="#FilteredCopyList" If="Exists('Engine/Binaries/Mac/UE4Editor.app/Contents/MacOS/NotForLicensees/UE4Editor-DDCUtils.dylib')"/>
+			<Tag Files="Engine/Binaries/Mac/NotForLicensees/ShaderCompileWorker.modules" With="#FilteredCopyList" If="Exists('Engine/Build/NotForLicensees/EpicInternal.txt')"/>
+			<Tag Files="Engine/Binaries/Mac/NotForLicensees/ShaderCompileWorker-DDCUtils.dylib" With="#FilteredCopyList" If="Exists('Engine/Build/NotForLicensees/EpicInternal.txt')"/>
+			<Tag Files="Engine/Binaries/Mac/UE4Editor.app/Contents/MacOS/NotForLicensees/UE4Editor.modules" With="#FilteredCopyList" If="Exists('Engine/Build/NotForLicensees/EpicInternal.txt')"/>
+			<Tag Files="Engine/Binaries/Mac/UE4Editor.app/Contents/MacOS/NotForLicensees/UE4Editor-DDCUtils.dylib" With="#FilteredCopyList" If="Exists('Engine/Build/NotForLicensees/EpicInternal.txt')"/>
 			<!-- Copy everything to a temporary directory -->
 			<Copy From="$(RootDir)" To="$(DDCDir)" Files="#FilteredCopyList"/>
 			<Copy From="$(CsToolsDir)" To="$(DDCDir)" Files="#Build Tools CS Binaries"/>
@@ -620,8 +624,10 @@
 			<Tag Files="#ToCopy" Except="$(ConfidentialExceptions);.../Source/...;.../Intermediate/..." With="#FilteredCopyList"/>
 			<!-- Additional required files that would be removed by confidential filter -->
 			<Tag Files="Engine/Build/NotForLicensees/EpicInternal.txt" With="#FilteredCopyList" If="Exists('Engine/Build/NotForLicensees/EpicInternal.txt')"/>
-			<Tag Files="Engine/Binaries/Linux/NotForLicensees/UE4Editor.modules" With="#FilteredCopyList" If="Exists('Engine/Binaries/Linux/NotForLicensees/UE4Editor.modules')"/>
-			<Tag Files="Engine/Binaries/Linux/NotForLicensees/libUE4Editor-DDCUtils.so" With="#FilteredCopyList" If="Exists('Engine/Binaries/Linux/NotForLicensees/libUE4Editor-DDCUtils.so')"/>
+			<Tag Files="Engine/Binaries/Linux/NotForLicensees/ShaderCompileWorker.modules" With="#FilteredCopyList" If="Exists('Engine/Build/NotForLicensees/EpicInternal.txt')"/>
+			<Tag Files="Engine/Binaries/Linux/NotForLicensees/libShaderCompileWorker-DDCUtils.so" With="#FilteredCopyList" If="Exists('Engine/Build/NotForLicensees/EpicInternal.txt')"/>
+			<Tag Files="Engine/Binaries/Linux/NotForLicensees/UE4Editor.modules" With="#FilteredCopyList" If="Exists('Engine/Build/NotForLicensees/EpicInternal.txt')"/>
+			<Tag Files="Engine/Binaries/Linux/NotForLicensees/libUE4Editor-DDCUtils.so" With="#FilteredCopyList" If="Exists('Engine/Build/NotForLicensees/EpicInternal.txt')"/>
 			<!-- Copy everything to a temporary directory -->
 			<Copy From="$(RootDir)" To="$(DDCDir)" Files="#FilteredCopyList"/>
 			<Copy From="$(CsToolsDir)" To="$(DDCDir)" Files="#Build Tools CS Binaries"/>
diff --git a/Engine/Source/Programs/AutomationTool/BuildGraph/Condition.cs b/Engine/Source/Programs/AutomationTool/BuildGraph/Condition.cs
--- a/Engine/Source/Programs/AutomationTool/BuildGraph/Condition.cs
+++ b/Engine/Source/Programs/AutomationTool/BuildGraph/Condition.cs
@@ -227,7 +227,7 @@ namespace AutomationTool
 				// Check whether file or directory exists. Evaluate the argument as a subexpression.
 				Idx++;
 				string Argument = EvaluateScalar(Tokens, ref Idx);
-				Result = (File.Exists(Argument) || Directory.Exists(Argument))? "true" : "false";
+				Result = Exists(Argument)? "true" : "false";
 			}
 			else if(String.Compare(Tokens[Idx], "HasTrailingSlash", true) == 0 && Tokens[Idx + 1] == "(")
 			{
@@ -258,6 +258,24 @@ namespace AutomationTool
 			return Result;
 		}
 
+		/// <summary>
+		/// Checks whether a path exists
+		/// </summary>
+		/// <param name="Scalar">The path to check for</param>
+		/// <returns>True if the path exists, false otherwise.</returns>
+		static bool Exists(string Scalar)
+		{
+			try
+			{
+				string FullPath = Path.Combine(CommandUtils.RootDirectory.FullName, Scalar);
+				return CommandUtils.FileExists(FullPath) || CommandUtils.DirectoryExists(FullPath);
+			}
+			catch
+			{
+				return false;
+			}
+		}
+
 		/// <summary>
 		/// Converts a scalar to a boolean value.
 		/// </summary>
