Remove PS4 from installed build by default. Introduced by CL 4021738.

diff --git a/Engine/Build/InstalledEngineBuild.xml b/Engine/Build/InstalledEngineBuild.xml
--- a/Engine/Build/InstalledEngineBuild.xml
+++ b/Engine/Build/InstalledEngineBuild.xml
@@ -13,7 +13,7 @@
 	<Option Name="WithTVOS" Restrict="true|false" DefaultValue="true" Description="Include the tvOS target platform"/>
 	<Option Name="WithLinux" Restrict="true|false" DefaultValue="true" Description="Include the Linux target platform"/>
 	<Option Name="WithHTML5" Restrict="true|false" DefaultValue="true" Description="Include the HTML5 target platform"/>
-	<Option Name="WithPS4" Restrict="true|false" DefaultValue="true" Description="Include the PS4 target platform"/>
+	<Option Name="WithPS4" Restrict="true|false" DefaultValue="false" Description="Include the PS4 target platform"/>
 	<Option Name="WithXboxOne" Restrict="true|false" DefaultValue="false" Description="Include the XboxOne target platform"/>
 	<Option Name="WithSwitch" Restrict="true|false" DefaultValue="false" Description="Include the Switch target platform"/>
 	<Option Name="WithLumin" Restrict="true|false" DefaultValue="true" Description="Include the Lumin target platform"/>
