Merging //UE4/Dev-Main @ 4035348 to Dev-Core (//UE4/Dev-Core)

diff --git a/Engine/Build/InstalledEngineBuild.xml b/Engine/Build/InstalledEngineBuild.xml
--- a/Engine/Build/InstalledEngineBuild.xml
+++ b/Engine/Build/InstalledEngineBuild.xml
@@ -447,7 +447,7 @@
 			<!-- Copy Source and referenced libraries to a new location with Confidential folders removed-->
 			<Tag Files="Engine/Source/Programs/AutomationTool/..." Filter="*.csproj" Except="$(ConfidentialExceptions)"  With="#UAT Projects"/>
 			<CsCompile Project="#UAT Projects" Configuration="Development" Platform="AnyCPU" TagReferences="#UATReferences" EnumerateOnly="true"/>
-			<Tag Files="Engine/Binaries/DotNET/...;Engine/Source/Programs/...;Engine/Source/Editor/SwarmInterface/...;#UATReferences" Except="$(ConfidentialExceptions);AutomationTool.csproj.References" With="#RedistUATSource"/>
+			<Tag Files="Engine/Binaries/DotNET/...;Engine/Binaries/ThirdParty/VisualStudio/...;Engine/Source/Programs/...;Engine/Source/Editor/SwarmInterface/...;#UATReferences" Except="$(ConfidentialExceptions);AutomationTool.csproj.References" With="#RedistUATSource"/>
 			<Copy Files="#RedistUATSource" From="$(RootDir)" To="$(CsToolsDir)"/>
 			<Tag Files="$(CsToolsDir)/Engine/Source/Programs/AutomationTool/..." Filter="*.csproj" With="#RedistUATProjects"/>
 
diff --git a/Engine/Source/Programs/UnrealBuildTool/Platform/Windows/UEBuildWindows.cs b/Engine/Source/Programs/UnrealBuildTool/Platform/Windows/UEBuildWindows.cs
--- a/Engine/Source/Programs/UnrealBuildTool/Platform/Windows/UEBuildWindows.cs
+++ b/Engine/Source/Programs/UnrealBuildTool/Platform/Windows/UEBuildWindows.cs
@@ -597,7 +597,7 @@ namespace UnrealBuildTool
 				    {
 					    // VS2015 just installs one toolchain; use that.
 					    DirectoryReference InstallDir;
-					    if(TryReadInstallDirRegistryKey32("Microsoft\\VisualStudio\\SxS\\VC7", "14.0", out InstallDir))
+					    if(TryReadInstallDirRegistryKey32("Microsoft\\VisualStudio\\SxS\\VS7", "14.0", out InstallDir))
 					    {
 						    InstallDirs.Add(InstallDir);
 					    }
@@ -661,7 +661,7 @@ namespace UnrealBuildTool
 					    List<DirectoryReference> InstallDirs = FindVSInstallDirs(Compiler);
 					    foreach(DirectoryReference InstallDir in InstallDirs)
 					    {
-						    ToolChainVersionToDir[new VersionNumber(14, 0)] = InstallDir;
+						    ToolChainVersionToDir[new VersionNumber(14, 0)] = DirectoryReference.Combine(InstallDir, "VC");
 					    }
 				    }
 				    else if(Compiler == WindowsCompiler.VisualStudio2017)
diff --git a/Engine/Source/Programs/Windows/BootstrapPackagedGame/Private/BootstrapPackagedGame.h b/Engine/Source/Programs/Windows/BootstrapPackagedGame/Private/BootstrapPackagedGame.h
--- a/Engine/Source/Programs/Windows/BootstrapPackagedGame/Private/BootstrapPackagedGame.h
+++ b/Engine/Source/Programs/Windows/BootstrapPackagedGame/Private/BootstrapPackagedGame.h
@@ -4,7 +4,10 @@
 
 struct IUnknown;
 
+#pragma pack(push, 8)
 #include <windows.h>
+#pragma pack(pop)
+
 #include <tchar.h>
 #include <assert.h>
 #include <stdio.h>
