Update stream for analytics.

diff --git a/Engine/Build/Build.version b/Engine/Build/Build.version
--- a/Engine/Build/Build.version
+++ b/Engine/Build/Build.version
@@ -6,5 +6,5 @@
 	"CompatibleChangelist": 0,
 	"IsLicenseeVersion": 0,
 	"IsPromotedBuild": 0,
-	"BranchName": "UE4"
+	"BranchName": "++UE4+Release-4.20"
 }
