Disable building IOS in installed builds on Windows by default. We don't support this on the cross-compile toolchain.

diff --git a/Engine/Build/InstalledEngineBuild.xml b/Engine/Build/InstalledEngineBuild.xml
--- a/Engine/Build/InstalledEngineBuild.xml
+++ b/Engine/Build/InstalledEngineBuild.xml
@@ -3,14 +3,21 @@
 
 	<!-- Helper option to make installed build for host platform only -->
 	<Option Name="HostPlatformOnly" Restrict="true|false" DefaultValue="false" Description="A helper option to make an installed build for your host platform only, so that you don't have to disable each platform individually"/>
-	
+
+	<!-- Option to enable all target platforms by default, as we ship in the launcher. Note that this will require cross-compiling IOS on Windows. -->
+	<Option Name="AllPlatforms" Restrict="true|false" DefaultValue="false" Description="Include all target platforms by default"/>
+
+	<!-- Set the default values for whether to include various platforms based on the host platform. -->
+	<Property Name="DefaultWithIOS" Value="true"/>
+	<Property Name="DefaultWithIOS" Value="false" If="'$(HostPlatform)' != 'Mac' And !$(AllPlatforms)"/>
+
 	<!-- Setup default target platforms for Installed build if they haven't been specified on the commandline -->
 	<Option Name="WithWin64" Restrict="true|false" DefaultValue="true" Description="Include the Win64 target platform"/>
 	<Option Name="WithWin32" Restrict="true|false" DefaultValue="true" Description="Include the Win32 target platform"/>
 	<Option Name="WithMac" Restrict="true|false" DefaultValue="true" Description="Include the Mac target platform"/>
 	<Option Name="WithAndroid" Restrict="true|false" DefaultValue="true" Description="Include the Android target platform"/>
-	<Option Name="WithIOS" Restrict="true|false" DefaultValue="true" Description="Include the iOS target platform"/>
-	<Option Name="WithTVOS" Restrict="true|false" DefaultValue="true" Description="Include the tvOS target platform"/>
+	<Option Name="WithIOS" Restrict="true|false" DefaultValue="$(DefaultWithIOS)" Description="Include the iOS target platform"/>
+	<Option Name="WithTVOS" Restrict="true|false" DefaultValue="$(DefaultWithIOS)" Description="Include the tvOS target platform"/>
 	<Option Name="WithLinux" Restrict="true|false" DefaultValue="true" Description="Include the Linux target platform"/>
 	<Option Name="WithHTML5" Restrict="true|false" DefaultValue="true" Description="Include the HTML5 target platform"/>
 	<Option Name="WithPS4" Restrict="true|false" DefaultValue="false" Description="Include the PS4 target platform"/>
