Need dump_syms and BreakpadSymbolEncoder in installed builds

diff --git a/Engine/Build/InstalledEngineFilters.xml b/Engine/Build/InstalledEngineFilters.xml
--- a/Engine/Build/InstalledEngineFilters.xml
+++ b/Engine/Build/InstalledEngineFilters.xml
@@ -261,6 +261,10 @@
 			Engine/Build/BatchFiles/Linux/RunXBuild.sh
 			Engine/Build/BatchFiles/RunUAT.sh
 
+			<!-- Debug tools -->
+			Engine/Binaries/Linux/dump_syms
+			Engine/Binaries/Linux/BreakpadSymbolEncoder
+
 			<!-- Mono -->
 			Engine/Binaries/ThirdParty/Mono/...
 
