Prevent schema being added as a build product twice.

diff --git a/Engine/Build/InstalledEngineBuild.xml b/Engine/Build/InstalledEngineBuild.xml
--- a/Engine/Build/InstalledEngineBuild.xml
+++ b/Engine/Build/InstalledEngineBuild.xml
@@ -523,7 +523,7 @@
 			<Compile Target="ShaderCompileWorker" Configuration="Development" Platform="Linux" Tag="#Build Tools Linux"/>
 			<Compile Target="UnrealLightmass" Configuration="Development" Platform="Linux" Tag="#Build Tools Linux"/>
 
-			<Tag Files="$(RootDir)/Engine/Saved/UnrealBuildTool/BuildConfiguration.Schema.xsd" With="#Build Tools Linux"/>
+			<Tag Files="$(RootDir)/Engine/Saved/UnrealBuildTool/BuildConfiguration.Schema.xsd" With="#Build Tools Linux" If="'$(HostPlatform)' == 'Linux'"/>
 		</Node>
 	</Agent>
 
