Don't bother tagging the BuildConfiguration schema on Linux and Mac; we only include it in the distribution for Visual Studio anyway.

diff --git a/Engine/Build/InstalledEngineBuild.xml b/Engine/Build/InstalledEngineBuild.xml
--- a/Engine/Build/InstalledEngineBuild.xml
+++ b/Engine/Build/InstalledEngineBuild.xml
@@ -522,8 +522,6 @@
 			<Compile Target="CrashReportClient" Platform="Linux" Configuration="Shipping"/>
 			<Compile Target="ShaderCompileWorker" Configuration="Development" Platform="Linux" Tag="#Build Tools Linux"/>
 			<Compile Target="UnrealLightmass" Configuration="Development" Platform="Linux" Tag="#Build Tools Linux"/>
-
-			<Tag Files="$(RootDir)/Engine/Saved/UnrealBuildTool/BuildConfiguration.Schema.xsd" With="#Build Tools Linux" If="'$(HostPlatform)' == 'Linux'"/>
 		</Node>
 	</Agent>
 
@@ -539,8 +537,6 @@
 			<Compile Target="DsymExporter" Configuration="Development" Platform="Mac" Tag="#Build Tools Mac"/>
 			<Compile Target="UE4EditorServices" Configuration="Development" Platform="Mac" Tag="#Build Tools Mac"/>
 			<Compile Target="UnrealAtoS" Configuration="Development" Platform="Mac" Tag="#Build Tools Mac"/>
-
-			<Tag Files="$(RootDir)/Engine/Saved/UnrealBuildTool/BuildConfiguration.Schema.xsd" With="#Build Tools Mac"/>
 		</Node>
 	</Agent>
 
