Disable debug info for all target platforms in installed engine builds, to fix bloated install size. IOS and TVOS still have it enabled for now, due to errors when packaging.

diff --git a/Engine/Build/InstalledEngineBuild.xml b/Engine/Build/InstalledEngineBuild.xml
--- a/Engine/Build/InstalledEngineBuild.xml
+++ b/Engine/Build/InstalledEngineBuild.xml
@@ -97,8 +97,8 @@
 	<Property Name="LuminGPUArchitectures" Value="es2"/>
 	
 	<!-- Compile flags for Windows Game binaries -->
-	<Property Name="WindowsGameCompileArguments" Value="-precompile"/>
-	<Property Name="WindowsGameCompileArguments" Value="$(WindowsGameCompileArguments) -nodebuginfo" If="'$(WithFullDebugInfo)' == false"/>
+	<Property Name="TargetDebugInfoArg" Value="-nodebuginfo"/>
+	<Property Name="TargetDebugInfoArg" Value="" If="$(WithFullDebugInfo)"/>
 
 	<!-- Include the script that lists files and patterns used when creating an installed build -->
 	<Include Script="InstalledEngineFilters.xml"/>
@@ -211,7 +211,7 @@
 		<Node Name="Compile UE4Game Win64" Requires="Compile UnrealHeaderTool Win64" Produces="#UE4Game Win64;#UE4Game Win64 Unstripped;#UE4Game Win64 Stripped;#UE4Game Win64 Unsigned;#UE4Game Win64 Signed;#UE4Game Win64 Includes">
 		
 			<ForEach Name="Configuration" Values="$(GameConfigurations)">
-				<Compile Target="UE4Game" Platform="Win64" Configuration="$(Configuration)" Tag="#UE4Game Win64" Arguments="$(WindowsGameCompileArguments)"/>
+				<Compile Target="UE4Game" Platform="Win64" Configuration="$(Configuration)" Tag="#UE4Game Win64" Arguments="-precompile $(TargetDebugInfoArg)"/>
 			</ForEach>
 
 			<Do If="$(EmbedSrcSrvInfo)">
@@ -248,7 +248,7 @@
 		<Node Name="Compile UE4Game Win32" Requires="Compile UnrealHeaderTool Win64" Produces="#UE4Game Win32;#UE4Game Win32 Unstripped;#UE4Game Win32 Stripped;#UE4Game Win32 Unsigned;#UE4Game Win32 Signed;#UE4Game Win32 Includes">
 		
 			<ForEach Name="Configuration" Values="$(GameConfigurations)">
-				<Compile Target="UE4Game" Platform="Win32" Configuration="$(Configuration)" Tag="#UE4Game Win32" Arguments="$(WindowsGameCompileArguments)"/>
+				<Compile Target="UE4Game" Platform="Win32" Configuration="$(Configuration)" Tag="#UE4Game Win32" Arguments="-precompile $(TargetDebugInfoArg)"/>
 			</ForEach>
 
 			<Do If="$(EmbedSrcSrvInfo)">
@@ -286,21 +286,11 @@
 		<Node Name="Compile UE4Game Linux" Requires="Compile UnrealHeaderTool $(HostPlatform)" Produces="#UE4Game Linux;#UE4Game Linux Unstripped;#UE4Game Linux Stripped;#UE4Game Linux Includes">
 		
 			<ForEach Name="Configuration" Values="$(GameConfigurations)">
-				<Compile Target="UE4Game" Platform="Linux" Configuration="$(Configuration)" Tag="#UE4Game Linux" Arguments="-precompile"/>
+				<Compile Target="UE4Game" Platform="Linux" Configuration="$(Configuration)" Tag="#UE4Game Linux" Arguments="-precompile $(TargetDebugInfoArg)"/>
 			</ForEach>
 			
 			<Tag Files="#UE4Game Linux" Filter="$(LinuxStripFilter)" With="#UE4Game Linux Unstripped"/>
-
-			<Switch>
-				<Case If="'$(WithFullDebugInfo)' == true">
-					<!-- Copy files for local debug builds -->
-					<Copy Files="#UE4Game Linux Unstripped" From="$(RootDir)" To="$(SavedOutput)"/>
-				</Case>
-				<Default>
-					<!-- Otherwise, strip the files for distribution -->
-					<Strip BaseDir="$(RootDir)" OutputDir="$(SavedOutput)" Platform="Linux" Files="#UE4Game Linux Unstripped" Tag="#UE4Game Linux Stripped"/>
-				</Default>
-			</Switch>
+			<Strip BaseDir="$(RootDir)" OutputDir="$(SavedOutput)" Platform="Linux" Files="#UE4Game Linux Unstripped" Tag="#UE4Game Linux Stripped"/>
 
 			<!-- Tag the generated includes for this target -->
 			<Tag Files="Engine/Intermediate/Build/Linux/B4D820EA/UE4/Inc/...;Engine/Plugins/.../Intermediate/Build/Linux/B4D820EA/UE4/Inc/..." With="#UE4Game Linux Includes"/>
@@ -310,7 +300,7 @@
 		<Node Name="Compile UE4Game Android" Requires="Compile UnrealHeaderTool $(HostPlatform)" Produces="#UE4Game Android;#UE4Game Android Unstripped;#UE4Game Android Stripped;#UE4Game Android Includes">
 		
 			<ForEach Name="Configuration" Values="$(GameConfigurations)">
-				<Compile Target="UE4Game" Platform="Android" Configuration="$(Configuration)" Tag="#UE4Game Android" Arguments="-precompile -architectures=$(AndroidArchitectures) -gpuarchitectures=$(AndroidGPUArchitectures)"/>
+				<Compile Target="UE4Game" Platform="Android" Configuration="$(Configuration)" Tag="#UE4Game Android" Arguments="-precompile -architectures=$(AndroidArchitectures) -gpuarchitectures=$(AndroidGPUArchitectures) $(TargetDebugInfoArg)"/>
 			</ForEach>
 		
 			<Tag Files="#UE4Game Android" Filter="$(AndroidStripFilter)" With="#UE4Game Android Unstripped"/>
@@ -324,7 +314,7 @@
 		<Node Name="Compile UE4Game HTML5" Requires="Compile UnrealHeaderTool $(HostPlatform)" Produces="#UE4Game HTML5;#UE4Game HTML5 Includes">
 		
 			<ForEach Name="Configuration" Values="$(GameConfigurations)">
-				<Compile Target="UE4Game" Platform="HTML5" Configuration="$(Configuration)" Tag="#UE4Game HTML5" Arguments="-precompile"/>
+				<Compile Target="UE4Game" Platform="HTML5" Configuration="$(Configuration)" Tag="#UE4Game HTML5" Arguments="-precompile $(TargetDebugInfoArg)"/>
 			</ForEach>
 		
 			<!-- Tag the generated includes for this target -->
@@ -335,7 +325,7 @@
 		<Node Name="Compile UE4Game PS4" Requires="Compile UnrealHeaderTool Win64" Produces="#UE4Game PS4;#UE4Game PS4 Includes">
 		
 			<ForEach Name="Configuration" Values="$(GameConfigurations)">
-				<Compile Target="UE4Game" Platform="PS4" Configuration="$(Configuration)" Tag="#UE4Game PS4" Arguments="-precompile"/>
+				<Compile Target="UE4Game" Platform="PS4" Configuration="$(Configuration)" Tag="#UE4Game PS4" Arguments="-precompile $(TargetDebugInfoArg)"/>
 			</ForEach>
 
 			<!-- Tag the generated includes for this target -->
@@ -346,7 +336,7 @@
 		<Node Name="Compile UE4Game XboxOne" Requires="Compile UnrealHeaderTool Win64" Produces="#UE4Game XboxOne;#UE4Game XboxOne Includes">
 		
 			<ForEach Name="Configuration" Values="$(GameConfigurations)">
-				<Compile Target="UE4Game" Platform="XboxOne" Configuration="$(Configuration)" Tag="#UE4Game XboxOne" Arguments="-precompile"/>
+				<Compile Target="UE4Game" Platform="XboxOne" Configuration="$(Configuration)" Tag="#UE4Game XboxOne" Arguments="-precompile $(TargetDebugInfoArg)"/>
 			</ForEach>
 			
 			<!-- Tag the generated includes for this target -->
@@ -357,7 +347,7 @@
 		<Node Name="Compile UE4Game Switch" Requires="Compile UnrealHeaderTool Win64" Produces="#UE4Game Switch;#UE4Game Switch Includes" If="'$(WithSwitch)' == true">
 		
 			<ForEach Name="Configuration" Values="$(GameConfigurations)">
-				<Compile Target="UE4Game" Platform="Switch" Configuration="$(Configuration)" Tag="#UE4Game Switch" Arguments="-precompile"/>
+				<Compile Target="UE4Game" Platform="Switch" Configuration="$(Configuration)" Tag="#UE4Game Switch" Arguments="-precompile $(TargetDebugInfoArg)"/>
 			</ForEach>
 			
 			<!-- Tag the generated includes for this target -->
@@ -365,8 +355,10 @@
 		</Node>
 		
 		<Node Name="Compile UE4Game Lumin" Requires="Compile UnrealHeaderTool Win64" Produces="#UE4Game Lumin;#UE4Game Lumin Unstripped;#UE4Game Lumin Stripped;#UE4Game Lumin Includes" If="'$(WithLumin)' == true">
-			<Compile Target="UE4Game" Platform="Lumin" Configuration="Development" Tag="#UE4Game Lumin" Arguments="-precompile -architectures=$(AndroidArchitectures) -gpuarchitectures=$(LuminGPUArchitectures)" AllowParallelExecutor="false"/>
-			<Compile Target="UE4Game" Platform="Lumin" Configuration="Shipping" Tag="#UE4Game Lumin" Arguments="-precompile -architectures=$(AndroidArchitectures) -gpuarchitectures=$(LuminGPUArchitectures)" AllowParallelExecutor="false"/>
+
+			<ForEach Name="Configuration" Values="$(GameConfigurations)">
+				<Compile Target="UE4Game" Platform="Lumin" Configuration="$(Configuration)" Tag="#UE4Game Lumin" Arguments="-precompile -architectures=$(AndroidArchitectures) -gpuarchitectures=$(LuminGPUArchitectures) $(TargetDebugInfoArg)" AllowParallelExecutor="false"/>
+			</ForEach>
 			
 			<Tag Files="#UE4Game Lumin" Filter="$(LuminStripFilter)" With="#UE4Game Lumin Unstripped"/>
 			<Strip BaseDir="$(RootDir)" OutputDir="$(SavedOutput)" Platform="Lumin" Files="#UE4Game Lumin Unstripped" Tag="#UE4Game Lumin Stripped"/>
@@ -382,21 +374,11 @@
 		<Node Name="Compile UE4Game Mac" Requires="Compile UnrealHeaderTool Mac" Produces="#UE4Game Mac;#UE4Game Mac Unstripped;#UE4Game Mac Stripped;#UE4Game Mac Unsigned;#UE4Game Mac Signed;#UE4Game Mac Includes">
 		
 			<ForEach Name="Configuration" Values="$(GameConfigurations)">
-				<Compile Target="UE4Game" Platform="Mac" Configuration="$(Configuration)" Tag="#UE4Game Mac" Arguments="-precompile"/>
+				<Compile Target="UE4Game" Platform="Mac" Configuration="$(Configuration)" Tag="#UE4Game Mac" Arguments="-precompile $(TargetDebugInfoArg)"/>
 			</ForEach>
 			
 			<Tag Files="#UE4Game Mac" Filter="$(MacStripFilter)" With="#UE4Game Mac Unstripped"/>
-
-			<Switch>
-				<Case If="'$(WithFullDebugInfo)' == true">
-					<!-- Copy files for local debug builds -->
-					<Copy Files="#UE4Game Mac Unstripped" From="$(RootDir)" To="$(SavedOutput)"/>
-				</Case>
-				<Default>
-					<!-- Otherwise, strip the files for distribution -->
-					<Strip BaseDir="$(RootDir)" OutputDir="$(SavedOutput)" Platform="Mac" Files="#UE4Game Mac Unstripped" Tag="#UE4Game Mac Stripped"/>
-				</Default>
-			</Switch>
+			<Strip BaseDir="$(RootDir)" OutputDir="$(SavedOutput)" Platform="Mac" Files="#UE4Game Mac Unstripped" Tag="#UE4Game Mac Stripped"/>
 
 			<Do If="'$(SignExecutables)' == true">
 				<Tag Files="#UE4Game Mac" Filter="$(MacSignFilter)" With="#UE4Game Mac Unsigned"/>
