Disable debug info for IOS and TVOS in installed builds.

diff --git a/Engine/Build/InstalledEngineBuild.xml b/Engine/Build/InstalledEngineBuild.xml
--- a/Engine/Build/InstalledEngineBuild.xml
+++ b/Engine/Build/InstalledEngineBuild.xml
@@ -394,7 +394,7 @@
 		<Node Name="Compile UE4Game IOS" Requires="Compile UnrealHeaderTool Mac" Produces="#UE4Game IOS;#UE4Game IOS Unstripped;#UE4Game IOS Stripped;#UE4Game IOS Includes">
 		
 			<ForEach Name="Configuration" Values="$(GameConfigurations)">
-				<Compile Target="UE4Game" Platform="IOS" Configuration="$(Configuration)" Tag="#UE4Game IOS" Arguments="-precompile"/>
+				<Compile Target="UE4Game" Platform="IOS" Configuration="$(Configuration)" Tag="#UE4Game IOS" Arguments="-precompile $(TargetDebugInfoArg)"/>
 			</ForEach>
 
 			<Tag Files="#UE4Game IOS" Filter="$(IOSStripFilter)" With="#UE4Game IOS Unstripped"/>
@@ -408,7 +408,7 @@
 		<Node Name="Compile UE4Game TVOS" Requires="Compile UnrealHeaderTool Mac" Produces="#UE4Game TVOS;#UE4Game TVOS Unstripped;#UE4Game TVOS Stripped;#UE4Game TVOS Includes">
 		
 			<ForEach Name="Configuration" Values="$(GameConfigurations)">
-				<Compile Target="UE4Game" Platform="TVOS" Configuration="$(Configuration)" Tag="#UE4Game TVOS" Arguments="-precompile"/>
+				<Compile Target="UE4Game" Platform="TVOS" Configuration="$(Configuration)" Tag="#UE4Game TVOS" Arguments="-precompile $(TargetDebugInfoArg)"/>
 			</ForEach>
 
 			<Tag Files="#UE4Game TVOS" Filter="$(TVOSStripFilter)" With="#UE4Game TVOS Unstripped"/>
