Force Win32/Win64 game targets in the installed engine build to be compiled with VS2015, to fix missing symbols linking object files built with VS2017 with VS0215 MSVCRT.

diff --git a/Engine/Build/InstalledEngineBuild.xml b/Engine/Build/InstalledEngineBuild.xml
--- a/Engine/Build/InstalledEngineBuild.xml
+++ b/Engine/Build/InstalledEngineBuild.xml
@@ -37,6 +37,7 @@
 	<!-- The analytics type -->
 	<Option Name="AnalyticsTypeOverride" DefaultValue="" Description="Identifier for analytic events to send"/>
 
+	<!-- Whether to embed source server info (ie. P4 paths) into PDB files -->
 	<Option Name="EmbedSrcSrvInfo" Restrict="true|false" DefaultValue="false" Description="Whether to add Source indexing to Windows game apps so they can be added to a symbol server"/>
 
 	<!-- Which game configurations to include for packaged applications -->
@@ -44,7 +45,10 @@
 
 	<!-- Generate full debug info for binary editor and packaged application builds -->
 	<Option Name="WithFullDebugInfo" Restrict="true|false" DefaultValue="false" Description="Generate full debug info for binary editor and packaged application builds"/>
-	
+
+	<!-- Whether to force the VS2015 compiler for Windows target platofms. Object files built against the 2015 CRT are compatible with 2017, but not vice-versa. -->
+	<Option Name="ForceVS2015" Restrict="true|false" DefaultValue="false" Description="Whether to force use of the VS2015 compiler for Windows targets. Object files built against the VS2015 CRT are compatible with VS2017 CRT, but not vice versa."/>
+
 	<Do If="'$(HostPlatformOnly)' == true">
 		<!-- Activate correct Target Platforms for host -->
 		<Switch>
@@ -100,6 +104,10 @@
 	<Property Name="TargetDebugInfoArg" Value="-nodebuginfo"/>
 	<Property Name="TargetDebugInfoArg" Value="" If="$(WithFullDebugInfo)"/>
 
+	<!-- Compile flags for Windows targets -->
+	<Property Name="VSCompilerArg" Value=""/>
+	<Property Name="VSCompilerArg" Value="-2015" If="$(ForceVS2015)"/>
+	
 	<!-- Include the script that lists files and patterns used when creating an installed build -->
 	<Include Script="InstalledEngineFilters.xml"/>
 
@@ -211,7 +219,7 @@
 		<Node Name="Compile UE4Game Win64" Requires="Compile UnrealHeaderTool Win64" Produces="#UE4Game Win64;#UE4Game Win64 Unstripped;#UE4Game Win64 Stripped;#UE4Game Win64 Unsigned;#UE4Game Win64 Signed;#UE4Game Win64 Includes">
 		
 			<ForEach Name="Configuration" Values="$(GameConfigurations)">
-				<Compile Target="UE4Game" Platform="Win64" Configuration="$(Configuration)" Tag="#UE4Game Win64" Arguments="-precompile $(TargetDebugInfoArg)"/>
+				<Compile Target="UE4Game" Platform="Win64" Configuration="$(Configuration)" Tag="#UE4Game Win64" Arguments="-precompile $(VSCompilerArg) $(TargetDebugInfoArg)"/>
 			</ForEach>
 
 			<Do If="$(EmbedSrcSrvInfo)">
@@ -248,7 +256,7 @@
 		<Node Name="Compile UE4Game Win32" Requires="Compile UnrealHeaderTool Win64" Produces="#UE4Game Win32;#UE4Game Win32 Unstripped;#UE4Game Win32 Stripped;#UE4Game Win32 Unsigned;#UE4Game Win32 Signed;#UE4Game Win32 Includes">
 		
 			<ForEach Name="Configuration" Values="$(GameConfigurations)">
-				<Compile Target="UE4Game" Platform="Win32" Configuration="$(Configuration)" Tag="#UE4Game Win32" Arguments="-precompile $(TargetDebugInfoArg)"/>
+				<Compile Target="UE4Game" Platform="Win32" Configuration="$(Configuration)" Tag="#UE4Game Win32" Arguments="-precompile $(VSCompilerArg) $(TargetDebugInfoArg)"/>
 			</ForEach>
 
 			<Do If="$(EmbedSrcSrvInfo)">
