Force the Visual Studio interop library to be included in the installed build. This isn't technically necessary (it has EmbedInteropTypes=true, so gets compiled into the executable), but QA sometimes attempt to rebuild it in-place.

diff --git a/Engine/Build/InstalledEngineFilters.xml b/Engine/Build/InstalledEngineFilters.xml
--- a/Engine/Build/InstalledEngineFilters.xml
+++ b/Engine/Build/InstalledEngineFilters.xml
@@ -60,6 +60,9 @@
 	
 	<!-- Define Editor Filters -->
 	<Property Name="CopyEditorFilter">
+		<!-- This assembly is normally embedded into the UBT executable, but it can technically be rebuilt from an installed build -->
+		Engine/Binaries/ThirdParty/VisualStudio/...
+
 		<!-- In-editor documentation -->
 		Engine/Documentation/Source/Shared/...
 		Engine/Documentation/Extras/...
