Fix AutoSDK support for Lumin on Mac, and enable it as part of installed engine builds.

diff --git a/Engine/Build/InstalledEngineBuild.xml b/Engine/Build/InstalledEngineBuild.xml
--- a/Engine/Build/InstalledEngineBuild.xml
+++ b/Engine/Build/InstalledEngineBuild.xml
@@ -427,8 +427,6 @@
 		</Node>
 		
 		<!-- Downstream monolithics and tools Lumin-->
-		<!--	
-		Temporarily removing until we have MLSDK setup on build machines 
 		<Node Name="Compile UE4Game LuminMac" Requires="Compile UnrealHeaderTool Mac" Produces="#UE4Game LuminMac;#UE4Game LuminMac Unstripped;#UE4Game LuminMac Stripped;#UE4Game LuminMac Includes">
 			<Compile Target="UE4Game" Platform="Lumin" Configuration="Development" Tag="#UE4Game LuminMac" Arguments="-precompile -architectures=$(AndroidArchitectures) -gpuarchitectures=$(LuminGPUArchitectures)" AllowParallelExecutor="false"/>
 			<Compile Target="UE4Game" Platform="Lumin" Configuration="Shipping" Tag="#UE4Game LuminMac" Arguments="-precompile -architectures=$(AndroidArchitectures) -gpuarchitectures=$(LuminGPUArchitectures)" AllowParallelExecutor="false"/>
@@ -438,7 +436,6 @@
 
 			<Tag Files="Engine/Intermediate/Build/Lumin/UE4/Inc/...;Engine/Plugins/.../Intermediate/Build/Lumin/UE4/Inc/..." With="#UE4Game LuminMac Includes"/>
 		</Node>
-		-->
 	</Agent>
 
 	<Agent Name="Tools Group Win64" Type="CompileWin64;Win64">
@@ -827,10 +824,7 @@
 		<Property Name="FilterRequirements" Value="$(FilterRequirements);#UE4Game IOS;#UE4Game IOS Unstripped;#UE4Game IOS Stripped;#UE4Game IOS Includes" If="'$(WithIOS)' == true"/>
 		<Property Name="FilterRequirements" Value="$(FilterRequirements);#UE4Game TVOS;#UE4Game TVOS Unstripped;#UE4Game TVOS Stripped;#UE4Game TVOS Includes" If="'$(WithTVOS)' == true"/>
 		<Property Name="FilterRequirements" Value="$(FilterRequirements);#UE4Game HTML5;#UE4Game HTML5 Includes" If="'$(WithHTML5)' == true"/>
-		<!--	
-		Temporarily removing until we have MLSDK setup on build machines 
 		<Property Name="FilterRequirements" Value="$(FilterRequirements);#UE4Game LuminMac;#UE4Game LuminMac Includes;#UE4Game LuminMac Unstripped;#UE4Game LuminMac Stripped;#UE4Game LuminMac Includes" If="'$(WithLumin)' == true"/> 
-		-->
 		<Property Name="FilterRequirements" Value="$(FilterRequirements);#CompressedDDCMac" If="'$(WithDDC)' == true"/>
 
 		<!-- Feature packs -->
@@ -881,15 +875,12 @@
 				<Property Name="CopyInstalledFilter" Value="$(CopyInstalledFilter);$(CopyHTML5Filter)"/>
 				<Property Name="CopyInstalledExceptions" Value="$(CopyInstalledExceptions);$(CopyHTML5Exceptions)"/>
 			</Do>
-			<!--	
-			Temporarily removing until we have MLSDK setup on build machines 
 			<Do If="'$(WithLumin)' == true">
 				<Tag Files="#UE4Game LuminMac;#UE4Game LuminMac Includes" Except="#UE4Game LuminMac Unstripped" With="#Installed Mac"/>
 				<Tag Files="#UE4Game LuminMac Stripped" With="#Saved Output"/>
 				<Property Name="CopyInstalledFilter" Value="$(CopyInstalledFilter);$(CopyLuminFilter)"/>
 				<Property Name="CopyInstalledExceptions" Value="$(CopyInstalledExceptions);$(CopyLuminExceptions)"/>
 			</Do>
-			-->
 
 			<!-- Tag any dependencies from all previous build products -->
 			<Tag Files="#Installed Mac" Filter="*.target" With="#TargetReceipts"/>
diff --git a/Engine/Source/Programs/UnrealBuildTool/Configuration/UEBuildPlatformSDK.cs b/Engine/Source/Programs/UnrealBuildTool/Configuration/UEBuildPlatformSDK.cs
--- a/Engine/Source/Programs/UnrealBuildTool/Configuration/UEBuildPlatformSDK.cs
+++ b/Engine/Source/Programs/UnrealBuildTool/Configuration/UEBuildPlatformSDK.cs
@@ -326,12 +326,28 @@ namespace UnrealBuildTool
 		/// <param name="Hook">Hook type</param>
 		protected virtual string GetHookExecutableName(SDKHookType Hook)
 		{
-			if (Hook == SDKHookType.Uninstall)
+			if(BuildHostPlatform.Current.Platform == UnrealTargetPlatform.Win64)
 			{
-				return "unsetup.bat";
+				if (Hook == SDKHookType.Uninstall)
+				{
+					return "unsetup.bat";
+				}
+				else
+				{
+					return "setup.bat";
+				}
+			}
+			else
+			{
+				if (Hook == SDKHookType.Uninstall)
+				{
+					return "unsetup.sh";
+				}
+				else
+				{
+					return "setup.sh";
+				}
 			}
-
-			return "setup.bat";
 		}
 
 		/// <summary>
diff --git a/Engine/Source/Programs/UnrealBuildTool/Platform/Lumin/LuminToolChain.cs b/Engine/Source/Programs/UnrealBuildTool/Platform/Lumin/LuminToolChain.cs
--- a/Engine/Source/Programs/UnrealBuildTool/Platform/Lumin/LuminToolChain.cs
+++ b/Engine/Source/Programs/UnrealBuildTool/Platform/Lumin/LuminToolChain.cs
@@ -120,7 +120,10 @@ namespace UnrealBuildTool
 			ObjCopyPath = ToolsDict["OBJCOPY"];
 
 			// force the compiler to be executed through a command prompt instance
-			bExecuteCompilerThroughShell = true;
+			if(BuildHostPlatform.Current.Platform == UnrealTargetPlatform.Win64)
+			{
+				bExecuteCompilerThroughShell = true;
+			}
 
 			// toolchain params
 			ToolchainParamsArm64 = " --sysroot=\"" + Path.Combine(MLSDKPath, "lumin") + "\"";
