Fix attempting to build UHT for Mac before building IOS when using the remote toolchain to create an installed engine build.

diff --git a/Engine/Build/InstalledEngineBuild.xml b/Engine/Build/InstalledEngineBuild.xml
--- a/Engine/Build/InstalledEngineBuild.xml
+++ b/Engine/Build/InstalledEngineBuild.xml
@@ -158,7 +158,7 @@
 	<!-- Base editor and tools for compile on Mac -->
 	<Agent Name="Editor Mac" Type="CompileMac;Mac">
 		<Node Name="Compile UnrealHeaderTool Mac" Requires="Update Version Files">
-			<Compile Target="UnrealHeaderTool" Platform="Mac" Configuration="Development" Arguments="-precompile"/>
+			<Compile Target="UnrealHeaderTool" Platform="Mac" Configuration="Development" Arguments="-precompile" If="'$(HostPlatform)' != 'Win64'"/>
 		</Node>
 		<Node Name="Compile UE4Editor Mac" Requires="Compile UnrealHeaderTool Mac" Produces="#UE4Editor Mac;#UE4Editor Mac Unstripped;#UE4Editor Mac Stripped;#UE4Editor Mac Unsigned;#UE4Editor Mac Signed;#UE4Editor Mac Includes">
 			<Compile Target="UE4Editor" Platform="Mac" Configuration="Development" Tag="#UE4Editor Mac" Arguments="-precompile"/>
