Add the nDisplay template project and related programs to the installed engine build.

#rnx

diff --git a/Engine/Build/InstalledEngineBuild.xml b/Engine/Build/InstalledEngineBuild.xml
--- a/Engine/Build/InstalledEngineBuild.xml
+++ b/Engine/Build/InstalledEngineBuild.xml
@@ -480,6 +480,8 @@
 				<CsCompile Project="$(CsToolsDir)/Engine/Source/Programs/UnrealSwarm/SwarmCoordinator/SwarmCoordinator.csproj" Configuration="Development" Platform="AnyCPU" Tag="#Build Tools CS Binaries;#Sign Binaries" TagReferences="#Build Tools CS Binaries"/>
 				<CsCompile Project="$(CsToolsDir)/Engine/Source/Programs/UnrealSwarm/Agent/Agent.csproj" Configuration="Development" Platform="AnyCPU" Tag="#Build Tools CS Binaries;#Sign Binaries" TagReferences="#Build Tools CS Binaries"/>
 				<CsCompile Project="$(CsToolsDir)/Engine/Source/Editor/SwarmInterface/DotNET/SwarmInterface.csproj" Configuration="Development" Platform="AnyCPU" Tag="#Build Tools CS Binaries;#Sign Binaries" TagReferences="#Build Tools CS Binaries"/>
+				<CsCompile Project="$(CsToolsDir)/Engine/Source/Programs/nDisplayLauncher/nDisplayLauncher.csproj" Configuration="Development" Platform="AnyCPU" Tag="#Build Tools CS Binaries;#Sign Binaries" TagReferences="#Build Tools CS Binaries"/>
+				<CsCompile Project="$(CsToolsDir)/Engine/Source/Programs/nDisplayListener/nDisplayListener.csproj" Configuration="Development" Platform="AnyCPU" Tag="#Build Tools CS Binaries;#Sign Binaries" TagReferences="#Build Tools CS Binaries"/>
 				<CsCompile Project="$(CsToolsDir)/Engine/Source/Programs/PS4/PS4DevKitUtil/PS4DevKitUtil.csproj" Configuration="Development" Platform="AnyCPU" Tag="#Build Tools CS Binaries;#Sign Binaries" TagReferences="#Build Tools CS Binaries" If="'$(WithPS4)' == true"/>
 				<CsCompile Project="$(CsToolsDir)/Engine/Source/Programs/PS4/PS4SymbolTool/PS4SymbolTool.csproj" Configuration="Release" Platform="AnyCPU" Tag="#Build Tools CS Binaries;#Sign Binaries" TagReferences="#Build Tools CS Binaries" If="'$(WithPS4)' == true"/>
 				<!-- TODO: Change these projects so that they can be compiled on non-windows platforms (currently use WinForms) -->
diff --git a/Engine/Build/InstalledEngineFilters.xml b/Engine/Build/InstalledEngineFilters.xml
--- a/Engine/Build/InstalledEngineFilters.xml
+++ b/Engine/Build/InstalledEngineFilters.xml
@@ -135,7 +135,8 @@
 		Templates/TP_2DSideScrollerBP/...
 		Templates/TP_VehicleAdv/...
 		Templates/TP_VehicleAdvBP/...
-		<!-- shared template resources -->
+		
+		<!-- Shared template resources -->
 		Templates/TemplateResources/...
 
 		<!-- Build files -->
@@ -211,6 +212,9 @@
 			Engine/Extras/MayaVelocityGridExporter/...
 			Engine/Extras/UnrealVS/...
 			Engine/Extras/Redist/en-us/*
+			
+			<!-- nDisplay Template -->
+			Templates/TP_nDisplayBP/...
 		</Property>
 
 		<Property Name="CopyEditorExceptions" Value="$(CopyEditorExceptions)">
