Fix Lumin not being disabled when building installed engine with -set:HostPlatformOnly=true.

#rnx

diff --git a/Engine/Build/InstalledEngineBuild.xml b/Engine/Build/InstalledEngineBuild.xml
--- a/Engine/Build/InstalledEngineBuild.xml
+++ b/Engine/Build/InstalledEngineBuild.xml
@@ -53,21 +53,18 @@
 		<!-- Activate correct Target Platforms for host -->
 		<Switch>
 			<Case If="'$(HostPlatform)' == 'Win64'">
-				<Property Name="WithWin64" Value="true"/>
 				<Property Name="WithMac" Value="false"/>
 				<Property Name="WithLinux" Value="false"/>
 			</Case>
 			<Case If="'$(HostPlatform)' == 'Mac'">
 				<Property Name="WithWin64" Value="false"/>
 				<Property Name="WithWin32" Value="false"/>
-				<Property Name="WithMac" Value="true"/>
 				<Property Name="WithLinux" Value="false"/>
 			</Case>
 			<Case If="'$(HostPlatform)' == 'Linux'">
 				<Property Name="WithWin64" Value="false"/>
 				<Property Name="WithWin32" Value="false"/>
 				<Property Name="WithMac" Value="false"/>
-				<Property Name="WithLinux" Value="true"/>
 			</Case>
 			<Default>
 				<Error Message="$(HostPlatform) is not supported for making an installed build"/>
@@ -81,6 +78,7 @@
 		<Property Name="WithPS4" Value="false"/>
 		<Property Name="WithXboxOne" Value="false"/>
 		<Property Name="WithSwitch" Value="false"/>
+		<Property Name="WithLumin" Value="false"/>
 	</Do>
 
 	<!-- The local output directory -->
