Strip debug info from Lumin on Mac. Increasing size of installed build by 5GB.

diff --git a/Engine/Build/InstalledEngineBuild.xml b/Engine/Build/InstalledEngineBuild.xml
--- a/Engine/Build/InstalledEngineBuild.xml
+++ b/Engine/Build/InstalledEngineBuild.xml
@@ -426,9 +426,11 @@
 		
 		<!-- Downstream monolithics and tools Lumin-->
 		<Node Name="Compile UE4Game LuminMac" Requires="Compile UnrealHeaderTool Mac" Produces="#UE4Game LuminMac;#UE4Game LuminMac Unstripped;#UE4Game LuminMac Stripped;#UE4Game LuminMac Includes">
-			<Compile Target="UE4Game" Platform="Lumin" Configuration="Development" Tag="#UE4Game LuminMac" Arguments="-precompile -architectures=$(AndroidArchitectures) -gpuarchitectures=$(LuminGPUArchitectures)" AllowParallelExecutor="false"/>
-			<Compile Target="UE4Game" Platform="Lumin" Configuration="Shipping" Tag="#UE4Game LuminMac" Arguments="-precompile -architectures=$(AndroidArchitectures) -gpuarchitectures=$(LuminGPUArchitectures)" AllowParallelExecutor="false"/>
 
+			<ForEach Name="Configuration" Values="$(GameConfigurations)">
+				<Compile Target="UE4Game" Platform="Lumin" Configuration="$(Configuration)" Tag="#UE4Game LuminMac" Arguments="-precompile -architectures=$(AndroidArchitectures) -gpuarchitectures=$(LuminGPUArchitectures) $(TargetDebugInfoArg)" AllowParallelExecutor="false"/>
+			</ForEach>
+			
 			<Tag Files="#UE4Game LuminMac" Filter="$(LuminStripFilter)" With="#UE4Game LuminMac Unstripped"/>
 			<Strip BaseDir="$(RootDir)" OutputDir="$(SavedOutput)" Platform="Lumin" Files="#UE4Game LuminMac Unstripped" Tag="#UE4Game LuminMac Stripped"/>
 
