UAT: Fix problems making installed builds with Lumin enabled, when passing the -Clean argument to UAT.

* UE4Build was deleting all files in the manifest before running, even if they were output by UBT.
* Lumin was specifying AllowParallelExecutor="false" in InstalledEngineBuild.xml, which was causing each compile to be executed in isolation. Since modules are now compiled into object files separately to being linked into binaries, this was causing object files and .precompiled files to be deleted during the clean operation.
* Individual compile tasks can now override the global "clean' option.

#rnx

diff --git a/Engine/Build/InstalledEngineBuild.xml b/Engine/Build/InstalledEngineBuild.xml
--- a/Engine/Build/InstalledEngineBuild.xml
+++ b/Engine/Build/InstalledEngineBuild.xml
@@ -213,7 +213,7 @@
 		
 			<ForEach Name="Configuration" Values="$(GameConfigurations)">
 				<Compile Target="UE4Game" Platform="Win64" Configuration="$(Configuration)" Tag="#UE4Game Win64" Arguments="-precompile -allmodules -nolink $(VSCompilerArg) $(TargetDebugInfoArg)"/>
-				<Compile Target="UE4Game" Platform="Win64" Configuration="$(Configuration)" Tag="#UE4Game Win64" Arguments="-precompile $(VSCompilerArg) $(TargetDebugInfoArg)"/>
+				<Compile Target="UE4Game" Platform="Win64" Configuration="$(Configuration)" Tag="#UE4Game Win64" Arguments="-precompile $(VSCompilerArg) $(TargetDebugInfoArg)" Clean="false"/>
 			</ForEach>
 
 			<Do If="$(EmbedSrcSrvInfo)">
@@ -248,7 +248,7 @@
 		
 			<ForEach Name="Configuration" Values="$(GameConfigurations)">
 				<Compile Target="UE4Game" Platform="Win32" Configuration="$(Configuration)" Tag="#UE4Game Win32" Arguments="-precompile -allmodules -nolink $(VSCompilerArg) $(TargetDebugInfoArg)"/>
-				<Compile Target="UE4Game" Platform="Win32" Configuration="$(Configuration)" Tag="#UE4Game Win32" Arguments="-precompile $(VSCompilerArg) $(TargetDebugInfoArg)"/>
+				<Compile Target="UE4Game" Platform="Win32" Configuration="$(Configuration)" Tag="#UE4Game Win32" Arguments="-precompile $(VSCompilerArg) $(TargetDebugInfoArg)" Clean="false"/>
 			</ForEach>
 
 			<Do If="$(EmbedSrcSrvInfo)">
@@ -284,7 +284,7 @@
 		
 			<ForEach Name="Configuration" Values="$(GameConfigurations)">
 				<Compile Target="UE4Game" Platform="Linux" Configuration="$(Configuration)" Tag="#UE4Game Linux" Arguments="-precompile -allmodules -nolink $(TargetDebugInfoArg)"/>
-				<Compile Target="UE4Game" Platform="Linux" Configuration="$(Configuration)" Tag="#UE4Game Linux" Arguments="-precompile $(TargetDebugInfoArg)"/>
+				<Compile Target="UE4Game" Platform="Linux" Configuration="$(Configuration)" Tag="#UE4Game Linux" Arguments="-precompile $(TargetDebugInfoArg)" Clean="false"/>
 			</ForEach>
 			
 			<Tag Files="#UE4Game Linux" Filter="$(LinuxStripFilter)" With="#UE4Game Linux Unstripped"/>
@@ -297,7 +297,7 @@
 		
 			<ForEach Name="Configuration" Values="$(GameConfigurations)">
 				<Compile Target="UE4Game" Platform="Android" Configuration="$(Configuration)" Tag="#UE4Game Android" Arguments="-precompile -allmodules -nolink -architectures=$(AndroidArchitectures) -gpuarchitectures=$(AndroidGPUArchitectures) $(TargetDebugInfoArg)"/>
-				<Compile Target="UE4Game" Platform="Android" Configuration="$(Configuration)" Tag="#UE4Game Android" Arguments="-precompile -architectures=$(AndroidArchitectures) -gpuarchitectures=$(AndroidGPUArchitectures) $(TargetDebugInfoArg)"/>
+				<Compile Target="UE4Game" Platform="Android" Configuration="$(Configuration)" Tag="#UE4Game Android" Arguments="-precompile -architectures=$(AndroidArchitectures) -gpuarchitectures=$(AndroidGPUArchitectures) $(TargetDebugInfoArg)" Clean="false"/>
 			</ForEach>
 		
 			<Tag Files="#UE4Game Android" Filter="$(AndroidStripFilter)" With="#UE4Game Android Unstripped"/>
@@ -310,7 +310,7 @@
 		
 			<ForEach Name="Configuration" Values="$(GameConfigurations)">
 				<Compile Target="UE4Game" Platform="HTML5" Configuration="$(Configuration)" Tag="#UE4Game HTML5" Arguments="-precompile -allmodules -nolink $(TargetDebugInfoArg)"/>
-				<Compile Target="UE4Game" Platform="HTML5" Configuration="$(Configuration)" Tag="#UE4Game HTML5" Arguments="-precompile $(TargetDebugInfoArg)"/>
+				<Compile Target="UE4Game" Platform="HTML5" Configuration="$(Configuration)" Tag="#UE4Game HTML5" Arguments="-precompile $(TargetDebugInfoArg)" Clean="false"/>
 			</ForEach>
 		
 		</Node>
@@ -320,7 +320,7 @@
 		
 			<ForEach Name="Configuration" Values="$(GameConfigurations)">
 				<Compile Target="UE4Game" Platform="PS4" Configuration="$(Configuration)" Tag="#UE4Game PS4" Arguments="-precompile -allmodules -nolink $(TargetDebugInfoArg)"/>
-				<Compile Target="UE4Game" Platform="PS4" Configuration="$(Configuration)" Tag="#UE4Game PS4" Arguments="-precompile $(TargetDebugInfoArg)"/>
+				<Compile Target="UE4Game" Platform="PS4" Configuration="$(Configuration)" Tag="#UE4Game PS4" Arguments="-precompile $(TargetDebugInfoArg)" Clean="false"/>
 			</ForEach>
 
 		</Node>
@@ -330,7 +330,7 @@
 		
 			<ForEach Name="Configuration" Values="$(GameConfigurations)">
 				<Compile Target="UE4Game" Platform="XboxOne" Configuration="$(Configuration)" Tag="#UE4Game XboxOne" Arguments="-precompile -allmodules -nolink $(TargetDebugInfoArg)"/>
-				<Compile Target="UE4Game" Platform="XboxOne" Configuration="$(Configuration)" Tag="#UE4Game XboxOne" Arguments="-precompile $(TargetDebugInfoArg)"/>
+				<Compile Target="UE4Game" Platform="XboxOne" Configuration="$(Configuration)" Tag="#UE4Game XboxOne" Arguments="-precompile $(TargetDebugInfoArg)" Clean="false"/>
 			</ForEach>
 			
 		</Node>
@@ -340,7 +340,7 @@
 		
 			<ForEach Name="Configuration" Values="$(GameConfigurations)">
 				<Compile Target="UE4Game" Platform="Switch" Configuration="$(Configuration)" Tag="#UE4Game Switch" Arguments="-precompile -allmodules -nolink $(TargetDebugInfoArg)"/>
-				<Compile Target="UE4Game" Platform="Switch" Configuration="$(Configuration)" Tag="#UE4Game Switch" Arguments="-precompile $(TargetDebugInfoArg)"/>
+				<Compile Target="UE4Game" Platform="Switch" Configuration="$(Configuration)" Tag="#UE4Game Switch" Arguments="-precompile $(TargetDebugInfoArg)" Clean="false"/>
 			</ForEach>
 			
 		</Node>
@@ -348,8 +348,8 @@
 		<Node Name="Compile UE4Game Lumin" Requires="Compile UnrealHeaderTool Win64" Produces="#UE4Game Lumin;#UE4Game Lumin Unstripped;#UE4Game Lumin Stripped" If="'$(WithLumin)' == true">
 
 			<ForEach Name="Configuration" Values="$(GameConfigurations)">
-				<Compile Target="UE4Game" Platform="Lumin" Configuration="$(Configuration)" Tag="#UE4Game Lumin" Arguments="-precompile -allmodules -nolink -architectures=$(AndroidArchitectures) -gpuarchitectures=$(LuminGPUArchitectures) $(TargetDebugInfoArg)" AllowParallelExecutor="false"/>
-				<Compile Target="UE4Game" Platform="Lumin" Configuration="$(Configuration)" Tag="#UE4Game Lumin" Arguments="-precompile -architectures=$(AndroidArchitectures) -gpuarchitectures=$(LuminGPUArchitectures) $(TargetDebugInfoArg)" AllowParallelExecutor="false"/>
+				<Compile Target="UE4Game" Platform="Lumin" Configuration="$(Configuration)" Tag="#UE4Game Lumin" Arguments="-precompile -allmodules -nolink -architectures=$(AndroidArchitectures) -gpuarchitectures=$(LuminGPUArchitectures) $(TargetDebugInfoArg)"/>
+				<Compile Target="UE4Game" Platform="Lumin" Configuration="$(Configuration)" Tag="#UE4Game Lumin" Arguments="-precompile -architectures=$(AndroidArchitectures) -gpuarchitectures=$(LuminGPUArchitectures) $(TargetDebugInfoArg)" Clean="false"/>
 			</ForEach>
 			
 			<Tag Files="#UE4Game Lumin" Filter="$(LuminStripFilter)" With="#UE4Game Lumin Unstripped"/>
@@ -365,7 +365,7 @@
 		
 			<ForEach Name="Configuration" Values="$(GameConfigurations)">
 				<Compile Target="UE4Game" Platform="Mac" Configuration="$(Configuration)" Tag="#UE4Game Mac" Arguments="-precompile -allmodules -nolink $(TargetDebugInfoArg)"/>
-				<Compile Target="UE4Game" Platform="Mac" Configuration="$(Configuration)" Tag="#UE4Game Mac" Arguments="-precompile $(TargetDebugInfoArg)"/>
+				<Compile Target="UE4Game" Platform="Mac" Configuration="$(Configuration)" Tag="#UE4Game Mac" Arguments="-precompile $(TargetDebugInfoArg)" Clean="false"/>
 			</ForEach>
 			
 			<Tag Files="#UE4Game Mac" Filter="$(MacStripFilter)" With="#UE4Game Mac Unstripped"/>
@@ -384,7 +384,7 @@
 		
 			<ForEach Name="Configuration" Values="$(GameConfigurations)">
 				<Compile Target="UE4Game" Platform="IOS" Configuration="$(Configuration)" Tag="#UE4Game IOS" Arguments="-precompile -allmodules -nolink $(TargetDebugInfoArg)"/>
-				<Compile Target="UE4Game" Platform="IOS" Configuration="$(Configuration)" Tag="#UE4Game IOS" Arguments="-precompile $(TargetDebugInfoArg)"/>
+				<Compile Target="UE4Game" Platform="IOS" Configuration="$(Configuration)" Tag="#UE4Game IOS" Arguments="-precompile $(TargetDebugInfoArg)" Clean="false"/>
 			</ForEach>
 
 			<Tag Files="#UE4Game IOS" Filter="$(IOSStripFilter)" With="#UE4Game IOS Unstripped"/>
@@ -397,7 +397,7 @@
 		
 			<ForEach Name="Configuration" Values="$(GameConfigurations)">
 				<Compile Target="UE4Game" Platform="TVOS" Configuration="$(Configuration)" Tag="#UE4Game TVOS" Arguments="-precompile -allmodules -nolink $(TargetDebugInfoArg)"/>
-				<Compile Target="UE4Game" Platform="TVOS" Configuration="$(Configuration)" Tag="#UE4Game TVOS" Arguments="-precompile $(TargetDebugInfoArg)"/>
+				<Compile Target="UE4Game" Platform="TVOS" Configuration="$(Configuration)" Tag="#UE4Game TVOS" Arguments="-precompile $(TargetDebugInfoArg)" Clean="false"/>
 			</ForEach>
 
 			<Tag Files="#UE4Game TVOS" Filter="$(TVOSStripFilter)" With="#UE4Game TVOS Unstripped"/>
@@ -410,7 +410,7 @@
 
 			<ForEach Name="Configuration" Values="$(GameConfigurations)">
 				<Compile Target="UE4Game" Platform="Lumin" Configuration="$(Configuration)" Tag="#UE4Game LuminMac" Arguments="-precompile -allmodules -nolink -architectures=$(AndroidArchitectures) -gpuarchitectures=$(LuminGPUArchitectures)" AllowParallelExecutor="false"/>
-				<Compile Target="UE4Game" Platform="Lumin" Configuration="$(Configuration)" Tag="#UE4Game LuminMac" Arguments="-precompile -architectures=$(AndroidArchitectures) -gpuarchitectures=$(LuminGPUArchitectures)" AllowParallelExecutor="false"/>
+				<Compile Target="UE4Game" Platform="Lumin" Configuration="$(Configuration)" Tag="#UE4Game LuminMac" Arguments="-precompile -architectures=$(AndroidArchitectures) -gpuarchitectures=$(LuminGPUArchitectures)" AllowParallelExecutor="false" Clean="false"/>
 			</ForEach>
 
 			<Tag Files="#UE4Game LuminMac" Filter="$(LuminStripFilter)" With="#UE4Game LuminMac Unstripped"/>
diff --git a/Engine/Source/Programs/AutomationTool/AutomationUtils/UE4Build.cs b/Engine/Source/Programs/AutomationTool/AutomationUtils/UE4Build.cs
--- a/Engine/Source/Programs/AutomationTool/AutomationUtils/UE4Build.cs
+++ b/Engine/Source/Programs/AutomationTool/AutomationUtils/UE4Build.cs
@@ -26,13 +26,14 @@ namespace AutomationTool
 
 		public void PrepareBuildProduct(string File, bool DeleteProduct = true)
 		{
-			if (!BuildProductFiles.Contains(File))
-			{
-				if (DeleteProduct && DeleteBuildProducts)
-				{
-					CommandUtils.DeleteFile(File);
-				}
-			}
+			// Do not delete these files; they may have been generated by UBT
+			// if (!BuildProductFiles.Contains(File))
+			// {
+			//	if (DeleteProduct && DeleteBuildProducts)
+			//	{
+			//		CommandUtils.DeleteFile(File);
+			//	}
+			// }
 		}
 
 		public bool HasBuildProduct(string InFile)
@@ -603,6 +604,9 @@ namespace AutomationTool
 
 			/// Extra UBT args
 			public string UBTArgs;
+
+			/// Whether to clean this target
+			public bool? Clean;
 		}
 
 
@@ -1338,9 +1342,10 @@ namespace AutomationTool
 			{
 				string TaskFilePath = CommandUtils.CombinePaths(CommandUtils.CmdEnv.LogFolder, @"UAT_XGE.xml");
 
-				if (DeleteBuildProducts)
+				foreach (var Target in Agenda.Targets)
 				{
-					foreach (var Target in Agenda.Targets)
+					bool bClean = Target.Clean ?? DeleteBuildProducts;
+					if (bClean)
 					{
 						CleanWithUBT(Target.TargetName, Target.Platform, Target.Config.ToString(), Target.UprojectPath, bForceMonolithic, bForceNonUnity, bForceDebugInfo, Target.UBTArgs, bForceUnity);
 					}
@@ -1426,9 +1431,10 @@ namespace AutomationTool
 
 			if (!bUsedXGE)
 			{
-				if (DeleteBuildProducts)
+				foreach (var Target in Agenda.Targets)
 				{
-					foreach (var Target in Agenda.Targets)
+					bool bClean = Target.Clean ?? DeleteBuildProducts;
+					if (bClean)
 					{
 						CleanWithUBT(Target.TargetName, Target.Platform, Target.Config.ToString(), Target.UprojectPath, bForceMonolithic, bForceNonUnity, bForceDebugInfo, Target.UBTArgs, bForceUnity);
 					}
diff --git a/Engine/Source/Programs/AutomationTool/BuildGraph/Tasks/CompileTask.cs b/Engine/Source/Programs/AutomationTool/BuildGraph/Tasks/CompileTask.cs
--- a/Engine/Source/Programs/AutomationTool/BuildGraph/Tasks/CompileTask.cs
+++ b/Engine/Source/Programs/AutomationTool/BuildGraph/Tasks/CompileTask.cs
@@ -53,6 +53,12 @@ namespace AutomationTool
 		[TaskParameter(Optional = true)]
 		public bool AllowParallelExecutor = true;
 
+		/// <summary>
+		/// Whether to allow cleaning this target. If unspecified, targets are cleaned if the -Clean argument passed on the command line.
+		/// </summary>
+		[TaskParameter(Optional = true)]
+		public bool? Clean = null;
+
 		/// <summary>
 		/// Tag to be applied to build products of this task
 		/// </summary>
@@ -123,7 +129,7 @@ namespace AutomationTool
 			bAllowXGE &= Parameters.AllowXGE;
 			bAllowParallelExecutor &= Parameters.AllowParallelExecutor;
 
-			UE4Build.BuildTarget Target = new UE4Build.BuildTarget { TargetName = Parameters.Target, Platform = Parameters.Platform, Config = Parameters.Configuration, UBTArgs = "-nobuilduht " + (Parameters.Arguments ?? "") };
+			UE4Build.BuildTarget Target = new UE4Build.BuildTarget { TargetName = Parameters.Target, Platform = Parameters.Platform, Config = Parameters.Configuration, UBTArgs = "-nobuilduht " + (Parameters.Arguments ?? ""), Clean = Parameters.Clean };
 			if(!String.IsNullOrEmpty(Parameters.Tag))
 			{
 				TargetToTagName.Add(Target, Parameters.Tag);
