Add missing dependency on UnrealPak being compiled before building feature packs.

#rnx

diff --git a/Engine/Build/InstalledEngineBuild.xml b/Engine/Build/InstalledEngineBuild.xml
--- a/Engine/Build/InstalledEngineBuild.xml
+++ b/Engine/Build/InstalledEngineBuild.xml
@@ -474,24 +474,15 @@
 			<CsCompile Project="$(CsToolsDir)/Engine/Source/Programs/IOS/MobileDeviceInterface/MobileDeviceInterface.csproj" Configuration="Development" Platform="AnyCPU" Arguments="/verbosity:minimal /target:Rebuild" Tag="#Build Tools CS Binaries;#Sign Binaries" TagReferences="#Build Tools CS Binaries"/>
 			<CsCompile Project="$(CsToolsDir)/Engine/Source/Programs/IOS/DeploymentInterface/DeploymentInterface.csproj" Configuration="Development" Platform="AnyCPU" Arguments="/verbosity:minimal /target:Rebuild" Tag="#Build Tools CS Binaries;#Sign Binaries" TagReferences="#Build Tools CS Binaries"/>
 			<CsCompile Project="$(CsToolsDir)/Engine/Source/Programs/HTML5/HTML5LaunchHelper/HTML5LaunchHelper.csproj" Configuration="Development" Platform="AnyCPU" Arguments="/verbosity:minimal /target:Rebuild" Tag="#Build Tools CS Binaries;#Sign Binaries" TagReferences="#Build Tools CS Binaries"/>
-		
+
 			<!-- Sign the binaries -->
 			<Sign Files="#Sign Binaries" If="'$(SignExecutables)' == true"/>
 		</Node>
-		<Node Name="Make Feature Packs" Requires="Update Version Files">
-			<!--Delete any pre-existing feature packs-->
-			<Delete Files="FeaturePacks/..."/>
-			<ForEach Name="TemplateName" Values="FP_FirstPerson;FP_FirstPersonBP;TP_2DSideScroller;TP_2DSideScrollerBP;TP_Flying;TP_FlyingBP;TP_HandheldARBP;TP_Puzzle;TP_PuzzleBP;TP_Rolling;TP_RollingBP;TP_SideScroller;TP_SideScrollerBP;TP_ThirdPerson;TP_ThirdPersonBP;TP_TopDown;TP_TopDownBP;TP_TwinStick;TP_TwinStickBP;TP_Vehicle;TP_VehicleAdv;TP_VehicleAdvBP;TP_VehicleBP;TP_VirtualRealityBP">
-				<PakFile Output="FeaturePacks/$(TemplateName).upack" Files="" ResponseFile="Templates/$(TemplateName)/contents.txt"/>
-			</ForEach>
-			<ForEach Name="SampleName" Values="MobileStarterContent;StarterContent">
-				<PakFile Output="FeaturePacks/$(SampleName).upack" Files="" ResponseFile="Samples/$(SampleName)/contents.txt"/>
-			</ForEach>
-		</Node>
 		<Node Name="Build Tools Linux" Requires="Compile UnrealHeaderTool $(HostPlatform)">
 			<Compile Target="CrashReportClient" Platform="Linux" Configuration="Shipping"/>
 			<Compile Target="ShaderCompileWorker" Configuration="Development" Platform="Linux" Tag="#Build Tools Linux"/>
 			<Compile Target="UnrealLightmass" Configuration="Development" Platform="Linux" Tag="#Build Tools Linux"/>
+			<Compile Target="UnrealPak" Configuration="Development" Platform="Linux" Tag="#Build Tools Linux"/>
 		</Node>
 	</Agent>
 
@@ -510,6 +501,19 @@
 		</Node>
 	</Agent>
 
+	<Agent Name="FeaturePacks Group" Type="CompileWin64;Win64">
+		<Node Name="Make Feature Packs" Requires="Update Version Files;Build Tools $(HostPlatform)">
+			<!--Delete any pre-existing feature packs-->
+			<Delete Files="FeaturePacks/..."/>
+			<ForEach Name="TemplateName" Values="FP_FirstPerson;FP_FirstPersonBP;TP_2DSideScroller;TP_2DSideScrollerBP;TP_Flying;TP_FlyingBP;TP_HandheldARBP;TP_Puzzle;TP_PuzzleBP;TP_Rolling;TP_RollingBP;TP_SideScroller;TP_SideScrollerBP;TP_ThirdPerson;TP_ThirdPersonBP;TP_TopDown;TP_TopDownBP;TP_TwinStick;TP_TwinStickBP;TP_Vehicle;TP_VehicleAdv;TP_VehicleAdvBP;TP_VehicleBP;TP_VirtualRealityBP">
+				<PakFile Output="FeaturePacks/$(TemplateName).upack" Files="" ResponseFile="Templates/$(TemplateName)/contents.txt"/>
+			</ForEach>
+			<ForEach Name="SampleName" Values="MobileStarterContent;StarterContent">
+				<PakFile Output="FeaturePacks/$(SampleName).upack" Files="" ResponseFile="Samples/$(SampleName)/contents.txt"/>
+			</ForEach>
+		</Node>
+	</Agent>
+
 	<Property Name="DDCPlatformsWin64" Value="Windows"/>
 	<Property Name="DDCPlatformsWin64" Value="$(DDCPlatformsWin64)+WindowsNoEditor" If="'$(WithWin64)' == true Or '$(WithWin32)' == true"/>
 	<Do If="'$(HostPlatformDDCOnly)' == false">
