Use the -allmodules option when building UHT for the installed engine build.

#rnx

diff --git a/Engine/Build/InstalledEngineBuild.xml b/Engine/Build/InstalledEngineBuild.xml
--- a/Engine/Build/InstalledEngineBuild.xml
+++ b/Engine/Build/InstalledEngineBuild.xml
@@ -116,7 +116,7 @@
 			<SetVersion Change="$(Change)" Branch="$(EscapedBranch)" If="$(IsBuildMachine)"/>
 		</Node>
 		<Node Name="Compile UnrealHeaderTool Win64" Requires="Update Version Files">
-			<Compile Target="UnrealHeaderTool" Platform="Win64" Configuration="Development" Arguments="-precompile"/>
+			<Compile Target="UnrealHeaderTool" Platform="Win64" Configuration="Development" Arguments="-precompile -allmodules"/>
 		</Node>
 		<Node Name="Compile UE4Editor Win64" Requires="Compile UnrealHeaderTool Win64" Produces="#UE4Editor Win64;#UE4Editor Win64 Unstripped;#UE4Editor Win64 Stripped;#UE4Editor Win64 Unsigned;#UE4Editor Win64 Signed">
 			<Compile Target="UE4Editor" Platform="Win64" Configuration="DebugGame" Tag="#UE4Editor Win64" Arguments="-precompile -allmodules"/>
