Prevent files being copied to staging folders if the sync workspace is in a directory called "source".

#rnx

diff --git a/Engine/Build/InstalledEngineBuild.xml b/Engine/Build/InstalledEngineBuild.xml
--- a/Engine/Build/InstalledEngineBuild.xml
+++ b/Engine/Build/InstalledEngineBuild.xml
@@ -550,7 +550,7 @@
 			<Tag Files="Engine\Binaries\DotNET\OneSky.dll" With="#ToCopy"/>
 			<Tag Files="Templates\TemplateResources\..." With="#ToCopy"/>
 			<!-- Filter out the files not needed to build DDC. Removing confidential folders can affect DDC keys, so we want to be sure that we're making DDC with a build that can use it. -->
-			<Tag Files="#ToCopy" Except="$(ConfidentialExceptions);.../Source/...;.../Intermediate/..." With="#FilteredCopyList"/>
+			<Tag Files="#ToCopy" Except="$(ConfidentialExceptions);$(RootDir)/.../Source/...;$(RootDir)/.../Intermediate/..." With="#FilteredCopyList"/>
 			<!-- Additional required files that would be removed by confidential filter. Note that conditions are evaluated before copying build products, so check against EpicInternal.txt -->
 			<Tag Files="Engine\Build\NotForLicensees\EpicInternal.txt" With="#FilteredCopyList" If="Exists('Engine\Build\NotForLicensees\EpicInternal.txt')"/>
 			<Tag Files="Engine\Binaries\Win64\NotForLicensees\ShaderCompileWorker.modules" With="#FilteredCopyList" If="Exists('Engine\Build\NotForLicensees\EpicInternal.txt')"/>
@@ -595,7 +595,7 @@
 			<Tag Files="Engine/Binaries/DotNET/OneSky.dll" With="#ToCopy"/>
 			<Tag Files="Templates/TemplateResources/..." With="#ToCopy"/>
 			<!-- Filter out the files not needed to build DDC. Removing confidential folders can affect DDC keys, so we want to be sure that we're making DDC with a build that can use it. -->
-			<Tag Files="#ToCopy" Except="$(ConfidentialExceptions);.../Source/...;.../Intermediate/..." With="#FilteredCopyList"/>
+			<Tag Files="#ToCopy" Except="$(ConfidentialExceptions);$(RootDir)/.../Source/...;$(RootDir)/.../Intermediate/..." With="#FilteredCopyList"/>
 			<!-- Additional required files that would be removed by confidential filter -->
 			<Tag Files="Engine/Build/NotForLicensees/EpicInternal.txt" With="#FilteredCopyList" If="Exists('Engine/Build/NotForLicensees/EpicInternal.txt')"/>
 			<Tag Files="Engine/Binaries/Mac/NotForLicensees/ShaderCompileWorker.modules" With="#FilteredCopyList" If="Exists('Engine/Build/NotForLicensees/EpicInternal.txt')"/>
@@ -639,7 +639,7 @@
 			<Tag Files="Engine/Binaries/DotNET/OneSky.dll" With="#ToCopy"/>
 			<Tag Files="Templates/TemplateResources/..." With="#ToCopy"/>
 			<!-- Filter out the files not needed to build DDC. Removing confidential folders can affect DDC keys, so we want to be sure that we're making DDC with a build that can use it. -->
-			<Tag Files="#ToCopy" Except="$(ConfidentialExceptions);.../Source/...;.../Intermediate/..." With="#FilteredCopyList"/>
+			<Tag Files="#ToCopy" Except="$(ConfidentialExceptions);$(RootDir)/.../Source/...;$(RootDir)/.../Intermediate/..." With="#FilteredCopyList"/>
 			<!-- Additional required files that would be removed by confidential filter -->
 			<Tag Files="Engine/Build/NotForLicensees/EpicInternal.txt" With="#FilteredCopyList" If="Exists('Engine/Build/NotForLicensees/EpicInternal.txt')"/>
 			<!-- When running on Linux -->
