Add resource files to the list of files to be copied on Windows.

#rnx

diff --git a/Engine/Build/InstalledEngineFilters.xml b/Engine/Build/InstalledEngineFilters.xml
--- a/Engine/Build/InstalledEngineFilters.xml
+++ b/Engine/Build/InstalledEngineFilters.xml
@@ -291,6 +291,9 @@
 
 	<!-- Win64 -->
 	<Property Name="CopyWin64Filter">
+		<!-- Default resource files -->
+		Engine/Build/Windows/Resources/...
+
 		<!-- App local dependencies -->
 		Engine/Binaries/ThirdParty/AppLocalDependencies/...
 	</Property>
