Fixed path for Mac DDCUtils files in InstalledEngineBuild.xml

diff --git a/Engine/Build/InstalledEngineBuild.xml b/Engine/Build/InstalledEngineBuild.xml
--- a/Engine/Build/InstalledEngineBuild.xml
+++ b/Engine/Build/InstalledEngineBuild.xml
@@ -600,8 +600,8 @@
 			<Tag Files="Engine/Build/NotForLicensees/EpicInternal.txt" With="#FilteredCopyList" If="Exists('Engine/Build/NotForLicensees/EpicInternal.txt')"/>
 			<Tag Files="Engine/Binaries/Mac/NotForLicensees/ShaderCompileWorker.modules" With="#FilteredCopyList" If="Exists('Engine/Build/NotForLicensees/EpicInternal.txt')"/>
 			<Tag Files="Engine/Binaries/Mac/NotForLicensees/ShaderCompileWorker-DDCUtils.dylib" With="#FilteredCopyList" If="Exists('Engine/Build/NotForLicensees/EpicInternal.txt')"/>
-			<Tag Files="Engine/Binaries/Mac/UE4Editor.app/Contents/MacOS/NotForLicensees/UE4Editor.modules" With="#FilteredCopyList" If="Exists('Engine/Build/NotForLicensees/EpicInternal.txt')"/>
-			<Tag Files="Engine/Binaries/Mac/UE4Editor.app/Contents/MacOS/NotForLicensees/UE4Editor-DDCUtils.dylib" With="#FilteredCopyList" If="Exists('Engine/Build/NotForLicensees/EpicInternal.txt')"/>
+			<Tag Files="Engine/Binaries/Mac/NotForLicensees/UE4Editor.modules" With="#FilteredCopyList" If="Exists('Engine/Build/NotForLicensees/EpicInternal.txt')"/>
+			<Tag Files="Engine/Binaries/Mac/NotForLicensees/UE4Editor-DDCUtils.dylib" With="#FilteredCopyList" If="Exists('Engine/Build/NotForLicensees/EpicInternal.txt')"/>
 			<!-- Copy everything to a temporary directory -->
 			<Copy From="$(RootDir)" To="$(DDCDir)" Files="#FilteredCopyList"/>
 			<Copy From="$(CsToolsDir)" To="$(DDCDir)" Files="#Build Tools CS Binaries"/>
