Add the -createstub argument to TVOS.

diff --git a/Engine/Build/InstalledEngineBuild.xml b/Engine/Build/InstalledEngineBuild.xml
--- a/Engine/Build/InstalledEngineBuild.xml
+++ b/Engine/Build/InstalledEngineBuild.xml
@@ -393,7 +393,7 @@
 		
 			<ForEach Name="Configuration" Values="$(GameConfigurations)">
 				<Compile Target="UE4Game" Platform="TVOS" Configuration="$(Configuration)" Tag="#UE4Game TVOS" Arguments="-precompile -allmodules -nolink $(TargetDebugInfoArg)"/>
-				<Compile Target="UE4Game" Platform="TVOS" Configuration="$(Configuration)" Tag="#UE4Game TVOS" Arguments="-precompile $(TargetDebugInfoArg)" Clean="false"/>
+				<Compile Target="UE4Game" Platform="TVOS" Configuration="$(Configuration)" Tag="#UE4Game TVOS" Arguments="-precompile -createstub $(TargetDebugInfoArg)" Clean="false"/>
 			</ForEach>
 
 			<Tag Files="#UE4Game TVOS" Filter="$(TVOSStripFilter)" With="#UE4Game TVOS Unstripped"/>
