Add DeploymentServerLauncher to the binary build

diff --git a/Engine/Build/InstalledEngineBuild.xml b/Engine/Build/InstalledEngineBuild.xml
--- a/Engine/Build/InstalledEngineBuild.xml
+++ b/Engine/Build/InstalledEngineBuild.xml
@@ -477,6 +477,7 @@
 			<CsCompile Project="$(CsToolsDir)/Engine/Source/Programs/IOS/DeploymentServer/DeploymentServer.csproj" Configuration="Development" Platform="AnyCPU" Arguments="/verbosity:minimal /target:Rebuild" Tag="#Build Tools CS Binaries;#Sign Binaries" TagReferences="#Build Tools CS Binaries"/>
 			<CsCompile Project="$(CsToolsDir)/Engine/Source/Programs/IOS/MobileDeviceInterface/MobileDeviceInterface.csproj" Configuration="Development" Platform="AnyCPU" Arguments="/verbosity:minimal /target:Rebuild" Tag="#Build Tools CS Binaries;#Sign Binaries" TagReferences="#Build Tools CS Binaries"/>
 			<CsCompile Project="$(CsToolsDir)/Engine/Source/Programs/IOS/DeploymentInterface/DeploymentInterface.csproj" Configuration="Development" Platform="AnyCPU" Arguments="/verbosity:minimal /target:Rebuild" Tag="#Build Tools CS Binaries;#Sign Binaries" TagReferences="#Build Tools CS Binaries"/>
+			<CsCompile Project="$(CsToolsDir)/Engine/Source/Programs/IOS/DeploymentServerLauncher/DeploymentServerLauncher.csproj" Configuration="Development" Platform="AnyCPU" Arguments="/verbosity:minimal /target:Rebuild" Tag="#Build Tools CS Binaries;#Sign Binaries" TagReferences="#Build Tools CS Binaries"/>
 			<CsCompile Project="$(CsToolsDir)/Engine/Source/Programs/HTML5/HTML5LaunchHelper/HTML5LaunchHelper.csproj" Configuration="Development" Platform="AnyCPU" Arguments="/verbosity:minimal /target:Rebuild" Tag="#Build Tools CS Binaries;#Sign Binaries" TagReferences="#Build Tools CS Binaries"/>
 
 			<!-- Sign the binaries -->
