Fix name of LiveCodingConsole target.

#rnx

diff --git a/Engine/Build/InstalledEngineBuild.xml b/Engine/Build/InstalledEngineBuild.xml
--- a/Engine/Build/InstalledEngineBuild.xml
+++ b/Engine/Build/InstalledEngineBuild.xml
@@ -427,7 +427,7 @@
 			<Compile Target="UnrealLightmass" Configuration="Development" Platform="Win64" Tag="#Build Tools Win64"/>
 			<Compile Target="UnrealPak" Configuration="Development" Platform="Win64" Tag="#Build Tools Win64"/>
 			<Compile Target="UnrealMultiUserServer" Configuration="Development" Platform="Win64" Tag="#Build Tools Win64"/>
-			<Compile Target="LiveCodeConsole" Configuration="Development" Platform="Win64" Tag="#Build Tools Win64"/>
+			<Compile Target="LiveCodingConsole" Configuration="Development" Platform="Win64" Tag="#Build Tools Win64"/>
 			<Compile Target="BootstrapPackagedGame" Configuration="Shipping" Platform="Win32" Tag="#Build Tools Win64"/>
 			<Compile Target="BootstrapPackagedGame" Configuration="Shipping" Platform="Win64" Tag="#Build Tools Win64"/>
 			<Compile Target="BuildPatchTool" Configuration="Shipping" Platform="Win64" Tag="#Build Tools Win64"/>
