Remove unneeded WMR dependencies

diff --git a/Engine/Build/InstalledEngineFilters.xml b/Engine/Build/InstalledEngineFilters.xml
--- a/Engine/Build/InstalledEngineFilters.xml
+++ b/Engine/Build/InstalledEngineFilters.xml
@@ -313,10 +313,6 @@
 		Engine/Source/ThirdParty/Oculus/LibOVRPlatform/*.tps
 		Engine/Source/ThirdParty/Oculus/OVRPlugin/OVRPlugin/lib/Win64/*.lib
 		Engine/Source/ThirdParty/Oculus/OVRPlugin/*.tps
-		
-		<!-- Windows Mixed Reality delay-load .dll related files -->
-		Engine/Plugins/Runtime/WindowsMixedReality/Binaries/ThirdParty/MixedRealityInteropLibrary/Win64/*.dll
-		Engine/Plugins/Runtime/WindowsMixedReality/Binaries/ThirdParty/MixedRealityInteropLibrary/Win64/*.pdb
 	</Property>
 	<Property Name="CopyWin64Exceptions">
 	</Property>
@@ -335,10 +331,6 @@
 		Engine/Source/ThirdParty/Oculus/LibOVRPlatform/*.tps
 		Engine/Source/ThirdParty/Oculus/OVRPlugin/OVRPlugin/lib/Win32/*.lib
 		Engine/Source/ThirdParty/Oculus/OVRPlugin/*.tps
-
-		<!-- Windows Mixed Reality delay-load .dll related files -->
-		Engine/Plugins/Runtime/WindowsMixedReality/Binaries/ThirdParty/MixedRealityInteropLibrary/Win32/*.dll
-		Engine/Plugins/Runtime/WindowsMixedReality/Binaries/ThirdParty/MixedRealityInteropLibrary/Win32/*.pdb
 	</Property>
 	<Property Name="CopyWin32Exceptions">
 	</Property>
