Bundle the Linux Toolchain for installed builds if it exists

diff --git a/Engine/Build/InstalledEngineBuild.xml b/Engine/Build/InstalledEngineBuild.xml
--- a/Engine/Build/InstalledEngineBuild.xml
+++ b/Engine/Build/InstalledEngineBuild.xml
@@ -89,6 +89,9 @@
 	<!-- Directory for storing a stripped version of UAT for installed builds -->
 	<Property Name="CsToolsDir" Value="$(SavedOutput)/CsTools"/>
 
+	<!-- Directory for toolchains -->
+	<Property Name="ToolchainDir" Value="$(RootDir)/Engine/Extras/ThirdPartyNotUE/SDKs"/>
+
 	<!-- Architectures that we build for Android -->
 	<Property Name="AndroidArchitectures" Value="armv7+arm64"/>
 
@@ -1024,6 +1027,10 @@
 			<Property Name="SavedDir" Value="$(SavedOutput)\Installed\Linux"/>
 			<Copy Files="#CompressedDDCLinux" From="$(SavedDir)" To="$(LocalInstalledDirLinux)" If="'$(WithDDC)' == true" Tag="#Installed Build Linux Files"/>
 
+			<!-- Copy the toolchain if it exists -->
+			<Property Name="LinuxToolchainDir" Value="$(ToolchainDir)/HostLinux"/>
+			<Copy Files="$(LinuxToolchainDir)/..." From="$(LinuxToolchainDir)" To="$(LocalInstalledDirLinux)/Engine/Extras/ThirdPartyNotUE/SDKs/HostLinux" If="Exists('$(LinuxToolchainDir)')" Tag="#Installed Build Linux Files"/>
+
 			<!-- List all of the platforms that we only allow for content projects -->
 			<Property Name="ContentOnlyPlatforms" Value=""/>
 
