Add an option to just make an installed build for the host platform, and allow overriding individual platforms on top of that.

#rnx

diff --git a/Engine/Build/InstalledEngineBuild.xml b/Engine/Build/InstalledEngineBuild.xml
--- a/Engine/Build/InstalledEngineBuild.xml
+++ b/Engine/Build/InstalledEngineBuild.xml
@@ -4,26 +4,41 @@
 	<!-- Helper option to make installed build for host platform only -->
 	<Option Name="HostPlatformOnly" Restrict="true|false" DefaultValue="false" Description="A helper option to make an installed build for your host platform only, so that you don't have to disable each platform individually"/>
 
+	<!-- Helper option to make an installed build with the host platform editor only -->
+	<Option Name="HostPlatformEditorOnly" Restrict="true|false" DefaultValue="false" Description="A helper option to make an installed build for your host platform only, so that you don't have to disable each platform individually"/>
+	
 	<!-- Option to enable all target platforms by default, as we ship in the launcher. Note that this will require cross-compiling IOS on Windows. -->
 	<Option Name="AllPlatforms" Restrict="true|false" DefaultValue="false" Description="Include all target platforms by default"/>
 
-	<!-- Set the default values for whether to include various platforms based on the host platform. -->
-	<Property Name="DefaultWithIOS" Value="true"/>
-	<Property Name="DefaultWithIOS" Value="false" If="'$(HostPlatform)' != 'Mac' And !$(AllPlatforms)"/>
+	<!-- Get default values for each platform being enabled -->
+	<Property Name="DefaultWithWindows" Value="true"/>
+	<Property Name="DefaultWithWindows" Value="false" If="$(HostPlatformEditorOnly) Or ($(HostPlatformOnly) And '$(HostPlatform)' != 'Win64')"/>
+
+	<Property Name="DefaultWithMac" Value="true"/>
+	<Property Name="DefaultWithMac" Value="false" If="$(HostPlatformEditorOnly) Or ($(HostPlatformOnly) And '$(HostPlatform)' != 'Mac')"/>
 
+	<Property Name="DefaultWithLinux" Value="true"/>
+	<Property Name="DefaultWithLinux" Value="false" If="$(HostPlatformEditorOnly) Or ($(HostPlatformOnly) And '$(HostPlatform)' != 'Linux')"/>
+
+	<Property Name="DefaultWithPlatform" Value="true"/>
+	<Property Name="DefaultWithPlatform" Value="false" If="$(HostPlatformEditorOnly)"/>
+
+	<Property Name="DefaultWithIOS" Value="$(DefaultWithPlatform)"/>
+	<Property Name="DefaultWithIOS" Value="false" If="'$(HostPlatform)' != 'Mac' And !$(AllPlatforms)"/>
+	
 	<!-- Setup default target platforms for Installed build if they haven't been specified on the commandline -->
-	<Option Name="WithWin64" Restrict="true|false" DefaultValue="true" Description="Include the Win64 target platform"/>
-	<Option Name="WithWin32" Restrict="true|false" DefaultValue="true" Description="Include the Win32 target platform"/>
-	<Option Name="WithMac" Restrict="true|false" DefaultValue="true" Description="Include the Mac target platform"/>
-	<Option Name="WithAndroid" Restrict="true|false" DefaultValue="true" Description="Include the Android target platform"/>
+	<Option Name="WithWin64" Restrict="true|false" DefaultValue="$(DefaultWithWindows)" Description="Include the Win64 target platform"/>
+	<Option Name="WithWin32" Restrict="true|false" DefaultValue="$(DefaultWithWindows)" Description="Include the Win32 target platform"/>
+	<Option Name="WithMac" Restrict="true|false" DefaultValue="$(DefaultWithMac)" Description="Include the Mac target platform"/>
+	<Option Name="WithAndroid" Restrict="true|false" DefaultValue="$(DefaultWithPlatform)" Description="Include the Android target platform"/>
 	<Option Name="WithIOS" Restrict="true|false" DefaultValue="$(DefaultWithIOS)" Description="Include the iOS target platform"/>
 	<Option Name="WithTVOS" Restrict="true|false" DefaultValue="$(DefaultWithIOS)" Description="Include the tvOS target platform"/>
-	<Option Name="WithLinux" Restrict="true|false" DefaultValue="true" Description="Include the Linux target platform"/>
-	<Option Name="WithHTML5" Restrict="true|false" DefaultValue="true" Description="Include the HTML5 target platform"/>
+	<Option Name="WithLinux" Restrict="true|false" DefaultValue="$(DefaultWithLinux)" Description="Include the Linux target platform"/>
+	<Option Name="WithHTML5" Restrict="true|false" DefaultValue="$(DefaultWithPlatform)" Description="Include the HTML5 target platform"/>
 	<Option Name="WithPS4" Restrict="true|false" DefaultValue="false" Description="Include the PS4 target platform"/>
 	<Option Name="WithXboxOne" Restrict="true|false" DefaultValue="false" Description="Include the XboxOne target platform"/>
 	<Option Name="WithSwitch" Restrict="true|false" DefaultValue="false" Description="Include the Switch target platform"/>
-	<Option Name="WithLumin" Restrict="true|false" DefaultValue="true" Description="Include the Lumin target platform"/>
+	<Option Name="WithLumin" Restrict="true|false" DefaultValue="$(DefaultWithPlatform)" Description="Include the Lumin target platform"/>
 
 	<!-- Whether to create a prebuilt DDC -->
 	<Option Name="WithDDC" Restrict="true|false" DefaultValue="true" Description="Build a standalone derived-data cache for the engine content and templates" />
@@ -45,39 +60,7 @@
 
 	<!-- Generate full debug info for binary editor and packaged application builds -->
 	<Option Name="WithFullDebugInfo" Restrict="true|false" DefaultValue="false" Description="Generate full debug info for binary editor and packaged application builds"/>
-
-	<Do If="'$(HostPlatformOnly)' == true">
-		<!-- Activate correct Target Platforms for host -->
-		<Switch>
-			<Case If="'$(HostPlatform)' == 'Win64'">
-				<Property Name="WithMac" Value="false"/>
-				<Property Name="WithLinux" Value="false"/>
-			</Case>
-			<Case If="'$(HostPlatform)' == 'Mac'">
-				<Property Name="WithWin64" Value="false"/>
-				<Property Name="WithWin32" Value="false"/>
-				<Property Name="WithLinux" Value="false"/>
-			</Case>
-			<Case If="'$(HostPlatform)' == 'Linux'">
-				<Property Name="WithWin64" Value="false"/>
-				<Property Name="WithWin32" Value="false"/>
-				<Property Name="WithMac" Value="false"/>
-			</Case>
-			<Default>
-				<Error Message="$(HostPlatform) is not supported for making an installed build"/>
-			</Default>
-		</Switch>
-		<!-- Other Target Platforms always disabled -->
-		<Property Name="WithAndroid" Value="false"/>
-		<Property Name="WithIOS" Value="false"/>
-		<Property Name="WithTVOS" Value="false"/>
-		<Property Name="WithHTML5" Value="false"/>
-		<Property Name="WithPS4" Value="false"/>
-		<Property Name="WithXboxOne" Value="false"/>
-		<Property Name="WithSwitch" Value="false"/>
-		<Property Name="WithLumin" Value="false"/>
-	</Do>
-
+	
 	<!-- The local output directory -->
 	<Property Name="LocalInstalledDir" Value="$(RootDir)/LocalBuilds/Engine/Windows"/>
 	<Property Name="LocalInstalledDirMac" Value="$(RootDir)/LocalBuilds/Engine/Mac"/>
