Add a separate node in the installed build script for building Win32 tools, so they don't have to be included if Win32 isn't.

diff --git a/Engine/Build/InstalledEngineBuild.xml b/Engine/Build/InstalledEngineBuild.xml
--- a/Engine/Build/InstalledEngineBuild.xml
+++ b/Engine/Build/InstalledEngineBuild.xml
@@ -404,17 +404,14 @@
 	<Agent Name="Tools Group Win64" Type="CompileWin64;Win64">
 		<Node Name="Build Tools Win64" Requires="Compile UnrealHeaderTool Win64">
 			<Compile Target="CrashReportClient" Configuration="Shipping" Platform="Win64" Tag="#Build Tools Win64"/>
-			<Compile Target="CrashReportClient" Configuration="Shipping" Platform="Win32" Tag="#Build Tools Win64"/>
 			<Compile Target="UnrealWatchdog" Configuration="Shipping" Platform="Win64" Tag="#Build Tools Win64"/>
 			<Compile Target="ShaderCompileWorker" Configuration="Development" Platform="Win64" Tag="#Build Tools Win64"/>
-			<Compile Target="UnrealCEFSubProcess" Configuration="Development" Platform="Win32" Tag="#Build Tools Win64"/>
 			<Compile Target="UnrealCEFSubProcess" Configuration="Development" Platform="Win64" Tag="#Build Tools Win64"/>
 			<Compile Target="UnrealFrontend" Configuration="Development" Platform="Win64" Tag="#Build Tools Win64"/>
 			<Compile Target="UnrealLightmass" Configuration="Development" Platform="Win64" Tag="#Build Tools Win64"/>
 			<Compile Target="UnrealPak" Configuration="Development" Platform="Win64" Tag="#Build Tools Win64"/>
 			<Compile Target="UnrealMultiUserServer" Configuration="Development" Platform="Win64" Tag="#Build Tools Win64"/>
 			<Compile Target="LiveCodingConsole" Configuration="Development" Platform="Win64" Tag="#Build Tools Win64"/>
-			<Compile Target="BootstrapPackagedGame" Configuration="Shipping" Platform="Win32" Tag="#Build Tools Win64"/>
 			<Compile Target="BootstrapPackagedGame" Configuration="Shipping" Platform="Win64" Tag="#Build Tools Win64"/>
 			<Compile Target="BuildPatchTool" Configuration="Shipping" Platform="Win64" Tag="#Build Tools Win64"/>
 			<Compile Target="WebRTCProxy" Configuration="Development" Platform="Win64" Tag="#Build Tools Win64"/>
@@ -477,6 +474,11 @@
 			<Compile Target="UnrealPak" Configuration="Development" Platform="Linux" Tag="#Build Tools Linux"/>
 			<Compile Target="UnrealMultiUserServer" Configuration="Development" Platform="Linux" Tag="#Build Tools Linux"/>
 		</Node>
+		<Node Name="Build Tools Win32" Requires="Compile UnrealHeaderTool Win64">
+			<Compile Target="CrashReportClient" Configuration="Shipping" Platform="Win32" Tag="#Build Tools Win32"/>
+			<Compile Target="UnrealCEFSubProcess" Configuration="Development" Platform="Win32" Tag="#Build Tools Win32"/>
+			<Compile Target="BootstrapPackagedGame" Configuration="Shipping" Platform="Win32" Tag="#Build Tools Win32"/>
+		</Node>
 	</Agent>
 
 	<Agent Name="Tools Group Mac" Type="CompileMac;Mac">
@@ -663,7 +665,7 @@
 		<Property Name="InstalledRequirements" Value="Compile UnrealHeaderTool Win64;#UE4Editor Win64;#UE4Editor Win64 Unstripped;#UE4Editor Win64 Stripped;#UE4Editor Win64 Unsigned;#UE4Editor Win64 Signed;#Build Tools Win64;#Build Tools CS Binaries"/>
 		<!-- Optional Target Platform requirements -->
 		<Property Name="InstalledRequirements" Value="$(InstalledRequirements);#UE4Game Win64;#UE4Game Win64 Unstripped;#UE4Game Win64 Stripped;#UE4Game Win64 Unsigned;#UE4Game Win64 Signed" If="'$(WithWin64)' == true"/>
-		<Property Name="InstalledRequirements" Value="$(InstalledRequirements);#UE4Game Win32;#UE4Game Win32 Unstripped;#UE4Game Win32 Stripped;#UE4Game Win32 Unsigned;#UE4Game Win32 Signed" If="'$(WithWin32)' == true"/>
+		<Property Name="InstalledRequirements" Value="$(InstalledRequirements);#UE4Game Win32;#UE4Game Win32 Unstripped;#UE4Game Win32 Stripped;#UE4Game Win32 Unsigned;#UE4Game Win32 Signed;#Build Tools Win32" If="'$(WithWin32)' == true"/>
 		<Property Name="InstalledRequirements" Value="$(InstalledRequirements);#UE4Game Android;#UE4Game Android Unstripped;#UE4Game Android Stripped" If="'$(WithAndroid)' == true"/>
 		<Property Name="InstalledRequirements" Value="$(InstalledRequirements);#UE4Game IOS;#UE4Game IOS Unstripped;#UE4Game IOS Stripped;#Compile UnrealHeaderTool Mac" If="'$(WithIOS)' == true"/>
 		<Property Name="InstalledRequirements" Value="$(InstalledRequirements);#UE4Game TVOS;#UE4Game TVOS Unstripped;#UE4Game TVOS Stripped" If="'$(WithTVOS)' == true"/>
