Add support for creating client and server targets from installed builds. Use the -set:WithClient=true and -set:WithServer=true arguments to BuildGraph to enable.

diff --git a/Engine/Build/InstalledEngineBuild.xml b/Engine/Build/InstalledEngineBuild.xml
--- a/Engine/Build/InstalledEngineBuild.xml
+++ b/Engine/Build/InstalledEngineBuild.xml
@@ -40,6 +40,10 @@
 	<Option Name="WithSwitch" Restrict="true|false" DefaultValue="false" Description="Include the Switch target platform"/>
 	<Option Name="WithLumin" Restrict="true|false" DefaultValue="$(DefaultWithPlatform)" Description="Include the Lumin target platform"/>
 
+	<!-- Set up target types -->
+	<Option Name="WithClient" Restrict="true|false" DefaultValue="false" Description="Include precompiled client targets"/>
+	<Option Name="WithServer" Restrict="true|false" DefaultValue="false" Description="Include precompiled server targets"/>
+	
 	<!-- Whether to create a prebuilt DDC -->
 	<Option Name="WithDDC" Restrict="true|false" DefaultValue="true" Description="Build a standalone derived-data cache for the engine content and templates" />
 
@@ -87,7 +91,14 @@
 
 	<!-- Compile flags for Windows targets -->
 	<Property Name="VSCompilerArg" Value=""/>
-	
+
+	<!-- Names of the client and server targets to build -->
+	<Property Name="OptionalClientTarget" Value=""/>
+	<Property Name="OptionalClientTarget" Value="UE4Client" If="$(WithClient)"/>
+
+	<Property Name="OptionalServerTarget" Value=""/>
+	<Property Name="OptionalServerTarget" Value="UE4Server" If="$(WithServer)"/>
+
 	<!-- Include the script that lists files and patterns used when creating an installed build -->
 	<Include Script="InstalledEngineFilters.xml"/>
 
@@ -189,13 +200,14 @@
 	</Agent>
 
 	<!-- Target Platforms normally built from a Win64 machine -->
-	<Agent Name="Target Platforms Win64" Type="CompileWin64;Win64">
-		<!-- Downstream monolithics and tools Win64 -->
+	<Agent Name="Target Win64" Type="CompileWin64;Win64">
 		<Node Name="Compile UE4Game Win64" Requires="Compile UnrealHeaderTool Win64" Produces="#UE4Game Win64;#UE4Game Win64 Unstripped;#UE4Game Win64 Stripped;#UE4Game Win64 Unsigned;#UE4Game Win64 Signed">
-		
-			<ForEach Name="Configuration" Values="$(GameConfigurations)">
-				<Compile Target="UE4Game" Platform="Win64" Configuration="$(Configuration)" Tag="#UE4Game Win64" Arguments="-precompile -allmodules -nolink $(VSCompilerArg) $(TargetDebugInfoArg)"/>
-				<Compile Target="UE4Game" Platform="Win64" Configuration="$(Configuration)" Tag="#UE4Game Win64" Arguments="-precompile $(VSCompilerArg) $(TargetDebugInfoArg)" Clean="false"/>
+
+			<ForEach Name="Target" Values="UE4Game;$(OptionalClientTarget);$(OptionalServerTarget)">
+				<ForEach Name="Configuration" Values="$(GameConfigurations)">
+					<Compile Target="$(Target)" Platform="Win64" Configuration="$(Configuration)" Tag="#UE4Game Win64" Arguments="-precompile -allmodules -nolink $(VSCompilerArg) $(TargetDebugInfoArg)"/>
+					<Compile Target="$(Target)" Platform="Win64" Configuration="$(Configuration)" Tag="#UE4Game Win64" Arguments="-precompile $(VSCompilerArg) $(TargetDebugInfoArg)" Clean="false"/>
+				</ForEach>
 			</ForEach>
 
 			<Do If="$(EmbedSrcSrvInfo)">
@@ -224,13 +236,15 @@
 				<Sign Files="#Game_ToSign_Win64" Tag="#UE4Game Win64 Signed"/>
 			</Do>
 		</Node>
-
-		<!-- Downstream monolithics and tools Win32 -->
+	</Agent>
+	<Agent Name="Target Win32" Type="CompileWin64;Win64">
 		<Node Name="Compile UE4Game Win32" Requires="Compile UnrealHeaderTool Win64" Produces="#UE4Game Win32;#UE4Game Win32 Unstripped;#UE4Game Win32 Stripped;#UE4Game Win32 Unsigned;#UE4Game Win32 Signed">
-		
-			<ForEach Name="Configuration" Values="$(GameConfigurations)">
-				<Compile Target="UE4Game" Platform="Win32" Configuration="$(Configuration)" Tag="#UE4Game Win32" Arguments="-precompile -allmodules -nolink $(VSCompilerArg) $(TargetDebugInfoArg)"/>
-				<Compile Target="UE4Game" Platform="Win32" Configuration="$(Configuration)" Tag="#UE4Game Win32" Arguments="-precompile $(VSCompilerArg) $(TargetDebugInfoArg)" Clean="false"/>
+
+			<ForEach Name="Target" Values="UE4Game;$(OptionalClientTarget)">
+				<ForEach Name="Configuration" Values="$(GameConfigurations)">
+					<Compile Target="$(Target)" Platform="Win32" Configuration="$(Configuration)" Tag="#UE4Game Win32" Arguments="-precompile -allmodules -nolink $(VSCompilerArg) $(TargetDebugInfoArg)"/>
+					<Compile Target="$(Target)" Platform="Win32" Configuration="$(Configuration)" Tag="#UE4Game Win32" Arguments="-precompile $(VSCompilerArg) $(TargetDebugInfoArg)" Clean="false"/>
+				</ForEach>
 			</ForEach>
 
 			<Do If="$(EmbedSrcSrvInfo)">
@@ -242,16 +256,16 @@
 
 			<Tag Files="#UE4Game Win32" Filter="$(Win32StripFilter)" With="#UE4Game Win32 Unstripped"/>
 
-				<Switch>
-					<Case If="'$(WithFullDebugInfo)' == true">
-						<!-- Copy files for local debug builds -->
-						<Copy Files="#UE4Game Win32 Unstripped" From="$(RootDir)" To="$(SavedOutput)"/>
-					</Case>
-					<Default>
-						<!-- Otherwise, strip the files for distribution -->
-						<Strip BaseDir="$(RootDir)" OutputDir="$(SavedOutput)" Platform="Win32" Files="#UE4Game Win32 Unstripped" Tag="#UE4Game Win32 Stripped"/>
-					</Default>
-				</Switch>
+			<Switch>
+				<Case If="'$(WithFullDebugInfo)' == true">
+					<!-- Copy files for local debug builds -->
+					<Copy Files="#UE4Game Win32 Unstripped" From="$(RootDir)" To="$(SavedOutput)"/>
+				</Case>
+				<Default>
+					<!-- Otherwise, strip the files for distribution -->
+					<Strip BaseDir="$(RootDir)" OutputDir="$(SavedOutput)" Platform="Win32" Files="#UE4Game Win32 Unstripped" Tag="#UE4Game Win32 Stripped"/>
+				</Default>
+			</Switch>
 
 
 			<Do If="'$(SignExecutables)' == true">
@@ -260,78 +274,91 @@
 				<Sign Files="#Game_ToSign_Win32" Tag="#UE4Game Win32 Signed"/>
 			</Do>
 		</Node>
-
-		<!-- Downstream monolithics and tools Linux -->
+	</Agent>
+	<Agent Name="Target Linux" Type="CompileWin64;Win64">
 		<Node Name="Compile UE4Game Linux" Requires="Compile UnrealHeaderTool $(HostPlatform)" Produces="#UE4Game Linux;#UE4Game Linux Unstripped;#UE4Game Linux Stripped">
-		
-			<ForEach Name="Configuration" Values="$(GameConfigurations)">
-				<Compile Target="UE4Game" Platform="Linux" Configuration="$(Configuration)" Tag="#UE4Game Linux" Arguments="-precompile -allmodules -nolink $(TargetDebugInfoArg)"/>
-				<Compile Target="UE4Game" Platform="Linux" Configuration="$(Configuration)" Tag="#UE4Game Linux" Arguments="-precompile $(TargetDebugInfoArg)" Clean="false"/>
+
+			<ForEach Name="Target" Values="UE4Game;$(OptionalClientTarget);$(OptionalServerTarget)">
+				<ForEach Name="Configuration" Values="$(GameConfigurations)">
+					<Compile Target="$(Target)" Platform="Linux" Configuration="$(Configuration)" Tag="#UE4Game Linux" Arguments="-precompile -allmodules -nolink $(TargetDebugInfoArg)"/>
+					<Compile Target="$(Target)" Platform="Linux" Configuration="$(Configuration)" Tag="#UE4Game Linux" Arguments="-precompile $(TargetDebugInfoArg)" Clean="false"/>
+				</ForEach>
 			</ForEach>
-			
+
 			<Tag Files="#UE4Game Linux" Filter="$(LinuxStripFilter)" With="#UE4Game Linux Unstripped"/>
 			<Strip BaseDir="$(RootDir)" OutputDir="$(SavedOutput)" Platform="Linux" Files="#UE4Game Linux Unstripped" Tag="#UE4Game Linux Stripped"/>
 
 		</Node>
-
-		<!-- Downstream monolithics and tools Android-->
+	</Agent>
+	<Agent Name="Target Android" Type="CompileWin64;Win64">
 		<Node Name="Compile UE4Game Android" Requires="Compile UnrealHeaderTool $(HostPlatform)" Produces="#UE4Game Android;#UE4Game Android Unstripped;#UE4Game Android Stripped">
-		
-			<ForEach Name="Configuration" Values="$(GameConfigurations)">
-				<Compile Target="UE4Game" Platform="Android" Configuration="$(Configuration)" Tag="#UE4Game Android" Arguments="-precompile -allmodules -nolink -architectures=$(AndroidArchitectures) -gpuarchitectures=$(AndroidGPUArchitectures) $(TargetDebugInfoArg)"/>
-				<Compile Target="UE4Game" Platform="Android" Configuration="$(Configuration)" Tag="#UE4Game Android" Arguments="-precompile -architectures=$(AndroidArchitectures) -gpuarchitectures=$(AndroidGPUArchitectures) $(TargetDebugInfoArg)" Clean="false"/>
+
+			<ForEach Name="Target" Values="UE4Game;$(OptionalClientTarget)">
+				<ForEach Name="Configuration" Values="$(GameConfigurations)">
+					<Compile Target="$(Target)" Platform="Android" Configuration="$(Configuration)" Tag="#UE4Game Android" Arguments="-precompile -allmodules -nolink -architectures=$(AndroidArchitectures) -gpuarchitectures=$(AndroidGPUArchitectures) $(TargetDebugInfoArg)"/>
+					<Compile Target="$(Target)" Platform="Android" Configuration="$(Configuration)" Tag="#UE4Game Android" Arguments="-precompile -architectures=$(AndroidArchitectures) -gpuarchitectures=$(AndroidGPUArchitectures) $(TargetDebugInfoArg)" Clean="false"/>
+				</ForEach>
 			</ForEach>
-		
+
 			<Tag Files="#UE4Game Android" Filter="$(AndroidStripFilter)" With="#UE4Game Android Unstripped"/>
 			<Strip BaseDir="$(RootDir)" OutputDir="$(SavedOutput)" Platform="Android" Files="#UE4Game Android Unstripped" Tag="#UE4Game Android Stripped"/>
 
 		</Node>
-
-		<!-- Downstream monolithics and tools HTML5-->
+	</Agent>
+	<Agent Name="Target HTML5" Type="CompileWin64;Win64">
 		<Node Name="Compile UE4Game HTML5" Requires="Compile UnrealHeaderTool $(HostPlatform)" Produces="#UE4Game HTML5">
-		
+
 			<ForEach Name="Configuration" Values="$(GameConfigurations)">
 				<Compile Target="UE4Game" Platform="HTML5" Configuration="$(Configuration)" Tag="#UE4Game HTML5" Arguments="-precompile -allmodules -nolink $(TargetDebugInfoArg)"/>
 				<Compile Target="UE4Game" Platform="HTML5" Configuration="$(Configuration)" Tag="#UE4Game HTML5" Arguments="-precompile $(TargetDebugInfoArg)" Clean="false"/>
 			</ForEach>
-		
-		</Node>
 
-		<!-- Downstream monolithics and tools PS4 -->
+		</Node>
+	</Agent>
+	<Agent Name="Target PS4" Type="CompileWin64;Win64">
 		<Node Name="Compile UE4Game PS4" Requires="Compile UnrealHeaderTool Win64" Produces="#UE4Game PS4">
-		
-			<ForEach Name="Configuration" Values="$(GameConfigurations)">
-				<Compile Target="UE4Game" Platform="PS4" Configuration="$(Configuration)" Tag="#UE4Game PS4" Arguments="-precompile -allmodules -nolink $(TargetDebugInfoArg)"/>
-				<Compile Target="UE4Game" Platform="PS4" Configuration="$(Configuration)" Tag="#UE4Game PS4" Arguments="-precompile $(TargetDebugInfoArg)" Clean="false"/>
+
+			<ForEach Name="Target" Values="UE4Game;$(OptionalClientTarget)">
+				<ForEach Name="Configuration" Values="$(GameConfigurations)">
+					<Compile Target="$(Target)" Platform="PS4" Configuration="$(Configuration)" Tag="#UE4Game PS4" Arguments="-precompile -allmodules -nolink $(TargetDebugInfoArg)"/>
+					<Compile Target="$(Target)" Platform="PS4" Configuration="$(Configuration)" Tag="#UE4Game PS4" Arguments="-precompile $(TargetDebugInfoArg)" Clean="false"/>
+				</ForEach>
 			</ForEach>
 
 		</Node>
-
-		<!-- Downstream monolithics and tools XboxOne -->
+	</Agent>
+	<Agent Name="Target XboxOne" Type="CompileWin64;Win64">
 		<Node Name="Compile UE4Game XboxOne" Requires="Compile UnrealHeaderTool Win64" Produces="#UE4Game XboxOne">
-		
-			<ForEach Name="Configuration" Values="$(GameConfigurations)">
-				<Compile Target="UE4Game" Platform="XboxOne" Configuration="$(Configuration)" Tag="#UE4Game XboxOne" Arguments="-precompile -allmodules -nolink $(TargetDebugInfoArg)"/>
-				<Compile Target="UE4Game" Platform="XboxOne" Configuration="$(Configuration)" Tag="#UE4Game XboxOne" Arguments="-precompile $(TargetDebugInfoArg)" Clean="false"/>
+
+			<ForEach Name="Target" Values="UE4Game;$(OptionalClientTarget)">
+				<ForEach Name="Configuration" Values="$(GameConfigurations)">
+					<Compile Target="$(Target)" Platform="XboxOne" Configuration="$(Configuration)" Tag="#UE4Game XboxOne" Arguments="-precompile -allmodules -nolink $(TargetDebugInfoArg)"/>
+					<Compile Target="$(Target)" Platform="XboxOne" Configuration="$(Configuration)" Tag="#UE4Game XboxOne" Arguments="-precompile $(TargetDebugInfoArg)" Clean="false"/>
+				</ForEach>
 			</ForEach>
-			
-		</Node>
 
-		<!-- Downstream monolithics Switch -->
+		</Node>
+	</Agent>
+	<Agent Name="Target Switch" Type="CompileWin64;Win64">
 		<Node Name="Compile UE4Game Switch" Requires="Compile UnrealHeaderTool Win64" Produces="#UE4Game Switch" If="'$(WithSwitch)' == true">
-		
-			<ForEach Name="Configuration" Values="$(GameConfigurations)">
-				<Compile Target="UE4Game" Platform="Switch" Configuration="$(Configuration)" Tag="#UE4Game Switch" Arguments="-precompile -allmodules -nolink $(TargetDebugInfoArg)"/>
-				<Compile Target="UE4Game" Platform="Switch" Configuration="$(Configuration)" Tag="#UE4Game Switch" Arguments="-precompile $(TargetDebugInfoArg)" Clean="false"/>
+
+			<ForEach Name="Target" Values="UE4Game;$(OptionalClientTarget)">
+				<ForEach Name="Configuration" Values="$(GameConfigurations)">
+					<Compile Target="$(Target)" Platform="Switch" Configuration="$(Configuration)" Tag="#UE4Game Switch" Arguments="-precompile -allmodules -nolink $(TargetDebugInfoArg)"/>
+					<Compile Target="$(Target)" Platform="Switch" Configuration="$(Configuration)" Tag="#UE4Game Switch" Arguments="-precompile $(TargetDebugInfoArg)" Clean="false"/>
+				</ForEach>
 			</ForEach>
-			
+
 		</Node>
-		
+	</Agent>
+	<Agent Name="Target Lumin" Type="CompileWin64;Win64">
 		<Node Name="Compile UE4Game Lumin" Requires="Compile UnrealHeaderTool Win64" Produces="#UE4Game Lumin;#UE4Game Lumin Unstripped;#UE4Game Lumin Stripped" If="'$(WithLumin)' == true">
 
-			<ForEach Name="Configuration" Values="$(GameConfigurations)">
-				<Compile Target="UE4Game" Platform="Lumin" Configuration="$(Configuration)" Tag="#UE4Game Lumin" Arguments="-precompile -allmodules -nolink -architectures=$(AndroidArchitectures) -gpuarchitectures=$(LuminGPUArchitectures) $(TargetDebugInfoArg)"/>
-				<Compile Target="UE4Game" Platform="Lumin" Configuration="$(Configuration)" Tag="#UE4Game Lumin" Arguments="-precompile -architectures=$(AndroidArchitectures) -gpuarchitectures=$(LuminGPUArchitectures) $(TargetDebugInfoArg)" Clean="false"/>
+			<ForEach Name="Target" Values="UE4Game;$(OptionalClientTarget)">
+				<ForEach Name="Configuration" Values="$(GameConfigurations)">
+					<Compile Target="$(Target)" Platform="Lumin" Configuration="$(Configuration)" Tag="#UE4Game Lumin" Arguments="-precompile -allmodules -nolink -architectures=$(AndroidArchitectures) -gpuarchitectures=$(LuminGPUArchitectures) $(TargetDebugInfoArg)"/>
+					<Compile Target="$(Target)" Platform="Lumin" Configuration="$(Configuration)" Tag="#UE4Game Lumin" Arguments="-precompile -architectures=$(AndroidArchitectures) -gpuarchitectures=$(LuminGPUArchitectures) $(TargetDebugInfoArg)" Clean="false"/>
+				</ForEach>
 			</ForEach>
 			
 			<Tag Files="#UE4Game Lumin" Filter="$(LuminStripFilter)" With="#UE4Game Lumin Unstripped"/>
@@ -341,15 +368,16 @@
 	</Agent>
 
 	<!-- Target Platforms normally built from a Mac machine -->
-	<Agent Name="Target Platforms Mac" Type="CompileMac;Mac">
-		<!-- Downstream monolithics and tools Mac -->
+	<Agent Name="Target Mac" Type="CompileMac;Mac">
 		<Node Name="Compile UE4Game Mac" Requires="Compile UnrealHeaderTool Mac" Produces="#UE4Game Mac;#UE4Game Mac Unstripped;#UE4Game Mac Stripped;#UE4Game Mac Unsigned;#UE4Game Mac Signed">
-		
-			<ForEach Name="Configuration" Values="$(GameConfigurations)">
-				<Compile Target="UE4Game" Platform="Mac" Configuration="$(Configuration)" Tag="#UE4Game Mac" Arguments="-precompile -allmodules -nolink $(TargetDebugInfoArg)"/>
-				<Compile Target="UE4Game" Platform="Mac" Configuration="$(Configuration)" Tag="#UE4Game Mac" Arguments="-precompile $(TargetDebugInfoArg)" Clean="false"/>
+
+			<ForEach Name="Target" Values="UE4Game;$(OptionalClientTarget);$(OptionalServerTarget)">
+				<ForEach Name="Configuration" Values="$(GameConfigurations)">
+					<Compile Target="$(Target)" Platform="Mac" Configuration="$(Configuration)" Tag="#UE4Game Mac" Arguments="-precompile -allmodules -nolink $(TargetDebugInfoArg)"/>
+					<Compile Target="$(Target)" Platform="Mac" Configuration="$(Configuration)" Tag="#UE4Game Mac" Arguments="-precompile $(TargetDebugInfoArg)" Clean="false"/>
+				</ForEach>
 			</ForEach>
-			
+
 			<Tag Files="#UE4Game Mac" Filter="$(MacStripFilter)" With="#UE4Game Mac Unstripped"/>
 			<Strip BaseDir="$(RootDir)" OutputDir="$(SavedOutput)" Platform="Mac" Files="#UE4Game Mac Unstripped" Tag="#UE4Game Mac Stripped"/>
 
@@ -360,39 +388,45 @@
 			</Do>
 
 		</Node>
-
-		<!-- Downstream monolithics and tools IOS-->
+	</Agent>
+	<Agent Name="Target IOS" Type="CompileMac;Mac">
 		<Node Name="Compile UE4Game IOS" Requires="Compile UnrealHeaderTool Mac" Produces="#UE4Game IOS;#UE4Game IOS Unstripped;#UE4Game IOS Stripped">
-		
-			<ForEach Name="Configuration" Values="$(GameConfigurations)">
-				<Compile Target="UE4Game" Platform="IOS" Configuration="$(Configuration)" Tag="#UE4Game IOS" Arguments="-precompile -allmodules -nolink $(TargetDebugInfoArg)"/>
-				<Compile Target="UE4Game" Platform="IOS" Configuration="$(Configuration)" Tag="#UE4Game IOS" Arguments="-precompile -createstub $(TargetDebugInfoArg)" Clean="false"/>
+
+			<ForEach Name="Target" Values="UE4Game;$(OptionalClientTarget)">
+				<ForEach Name="Configuration" Values="$(GameConfigurations)">
+					<Compile Target="$(Target)" Platform="IOS" Configuration="$(Configuration)" Tag="#UE4Game IOS" Arguments="-precompile -allmodules -nolink $(TargetDebugInfoArg)"/>
+					<Compile Target="$(Target)" Platform="IOS" Configuration="$(Configuration)" Tag="#UE4Game IOS" Arguments="-precompile -createstub $(TargetDebugInfoArg)" Clean="false"/>
+				</ForEach>
 			</ForEach>
 
 			<Tag Files="#UE4Game IOS" Filter="$(IOSStripFilter)" With="#UE4Game IOS Unstripped"/>
 			<Strip BaseDir="$(RootDir)" OutputDir="$(SavedOutput)" Platform="IOS" Files="#UE4Game IOS Unstripped" Tag="#UE4Game IOS Stripped"/>
 
 		</Node>
-
-		<!-- Downstream monolithics and tools TVOS-->
+	</Agent>
+	<Agent Name="Target TVOS" Type="CompileMac;Mac">
 		<Node Name="Compile UE4Game TVOS" Requires="Compile UnrealHeaderTool Mac" Produces="#UE4Game TVOS;#UE4Game TVOS Unstripped;#UE4Game TVOS Stripped">
-		
-			<ForEach Name="Configuration" Values="$(GameConfigurations)">
-				<Compile Target="UE4Game" Platform="TVOS" Configuration="$(Configuration)" Tag="#UE4Game TVOS" Arguments="-precompile -allmodules -nolink $(TargetDebugInfoArg)"/>
-				<Compile Target="UE4Game" Platform="TVOS" Configuration="$(Configuration)" Tag="#UE4Game TVOS" Arguments="-precompile -createstub $(TargetDebugInfoArg)" Clean="false"/>
+
+			<ForEach Name="Target" Values="UE4Game;$(OptionalClientTarget)">
+				<ForEach Name="Configuration" Values="$(GameConfigurations)">
+					<Compile Target="$(Target)" Platform="TVOS" Configuration="$(Configuration)" Tag="#UE4Game TVOS" Arguments="-precompile -allmodules -nolink $(TargetDebugInfoArg)"/>
+					<Compile Target="$(Target)" Platform="TVOS" Configuration="$(Configuration)" Tag="#UE4Game TVOS" Arguments="-precompile -createstub $(TargetDebugInfoArg)" Clean="false"/>
+				</ForEach>
 			</ForEach>
 
 			<Tag Files="#UE4Game TVOS" Filter="$(TVOSStripFilter)" With="#UE4Game TVOS Unstripped"/>
 			<Strip BaseDir="$(RootDir)" OutputDir="$(SavedOutput)" Platform="TVOS" Files="#UE4Game TVOS Unstripped" Tag="#UE4Game TVOS Stripped"/>
 
 		</Node>
-		
-		<!-- Downstream monolithics and tools Lumin-->
+	</Agent>
+	<Agent Name="Target LuminMac" Type="CompileMac;Mac">
 		<Node Name="Compile UE4Game LuminMac" Requires="Compile UnrealHeaderTool Mac" Produces="#UE4Game LuminMac;#UE4Game LuminMac Unstripped;#UE4Game LuminMac Stripped">
 
-			<ForEach Name="Configuration" Values="$(GameConfigurations)">
-				<Compile Target="UE4Game" Platform="Lumin" Configuration="$(Configuration)" Tag="#UE4Game LuminMac" Arguments="-precompile -allmodules -nolink -architectures=$(AndroidArchitectures) -gpuarchitectures=$(LuminGPUArchitectures) $(TargetDebugInfoArg)" AllowParallelExecutor="false"/>
-				<Compile Target="UE4Game" Platform="Lumin" Configuration="$(Configuration)" Tag="#UE4Game LuminMac" Arguments="-precompile -architectures=$(AndroidArchitectures) -gpuarchitectures=$(LuminGPUArchitectures) $(TargetDebugInfoArg)" AllowParallelExecutor="false" Clean="false"/>
+			<ForEach Name="Target" Values="UE4Game;$(OptionalClientTarget)">
+				<ForEach Name="Configuration" Values="$(GameConfigurations)">
+					<Compile Target="$(Target)" Platform="Lumin" Configuration="$(Configuration)" Tag="#UE4Game LuminMac" Arguments="-precompile -allmodules -nolink -architectures=$(AndroidArchitectures) -gpuarchitectures=$(LuminGPUArchitectures) $(TargetDebugInfoArg)" AllowParallelExecutor="false"/>
+					<Compile Target="$(Target)" Platform="Lumin" Configuration="$(Configuration)" Tag="#UE4Game LuminMac" Arguments="-precompile -architectures=$(AndroidArchitectures) -gpuarchitectures=$(LuminGPUArchitectures) $(TargetDebugInfoArg)" AllowParallelExecutor="false" Clean="false"/>
+				</ForEach>
 			</ForEach>
 
 			<Tag Files="#UE4Game LuminMac" Filter="$(LuminStripFilter)" With="#UE4Game LuminMac Unstripped"/>
diff --git a/Engine/Build/InstalledEngineFilters.xml b/Engine/Build/InstalledEngineFilters.xml
--- a/Engine/Build/InstalledEngineFilters.xml
+++ b/Engine/Build/InstalledEngineFilters.xml
@@ -143,10 +143,13 @@
 		Engine/Build/Build.version
 		Engine/Build/Target.cs.template
 	</Property>
-  <!-- Optional API docs -->
-  <Property Name="CopyEditorFilter" Value="$(CopyEditorFilter);Engine/Documentation/Builds/CodeAPI-HTML.tgz" If="Exists('Engine/Documentation/Builds/CodeAPI-HTML.tgz')"/>
-  <Property Name="CopyEditorFilter" Value="$(CopyEditorFilter);Engine/Documentation/Builds/BlueprintAPI-HTML.tgz" If="Exists('Engine/Documentation/Builds/BlueprintAPI-HTML.tgz')"/>
-  <!-- Don't add VR template for Mac and Linux -->
+	<!-- Optional target files -->
+	<Property Name="CopyEditorFilter" Value="$(CopyEditorFilter);Engine/Source/UE4Client.Target.cs" If="$(WithClient)"/>
+	<Property Name="CopyEditorFilter" Value="$(CopyEditorFilter);Engine/Source/UE4Server.Target.cs" If="$(WithServer)"/>
+	<!-- Optional API docs -->
+	<Property Name="CopyEditorFilter" Value="$(CopyEditorFilter);Engine/Documentation/Builds/CodeAPI-HTML.tgz" If="Exists('Engine/Documentation/Builds/CodeAPI-HTML.tgz')"/>
+	<Property Name="CopyEditorFilter" Value="$(CopyEditorFilter);Engine/Documentation/Builds/BlueprintAPI-HTML.tgz" If="Exists('Engine/Documentation/Builds/BlueprintAPI-HTML.tgz')"/>
+	<!-- Don't add VR template for Mac and Linux -->
 	<Property Name="CopyEditorFilter" Value="$(CopyEditorFilter);Templates/TP_VirtualRealityBP/..." If="'$(HostPlatform)' != 'Linux'"/>
 		<!-- Define Editor Exceptions -->
 	<Property Name="CopyEditorExceptions">
diff --git a/Engine/Source/Programs/AutomationTool/Scripts/FinalizeInstalledBuild.cs b/Engine/Source/Programs/AutomationTool/Scripts/FinalizeInstalledBuild.cs
--- a/Engine/Source/Programs/AutomationTool/Scripts/FinalizeInstalledBuild.cs
+++ b/Engine/Source/Programs/AutomationTool/Scripts/FinalizeInstalledBuild.cs
@@ -108,27 +108,31 @@ namespace AutomationTool
 					bool bCanBeDisplayed = ProjectType == EProjectType.Content;
 					foreach (UnrealTargetConfiguration CodeTargetConfiguration in Enum.GetValues(typeof(UnrealTargetConfiguration)))
 					{
-						// Need to check for development receipt as we use that for the Engine code in DebugGame
-						UnrealTargetConfiguration EngineConfiguration = (CodeTargetConfiguration == UnrealTargetConfiguration.DebugGame) ? UnrealTargetConfiguration.Development : CodeTargetConfiguration;
-						FileReference ReceiptFileName = TargetReceipt.GetDefaultPath(new DirectoryReference(OutputEnginePath), "UE4Game", CodeTargetPlatform, EngineConfiguration, Architecture);
-
-						if (FileReference.Exists(ReceiptFileName))
+						string[] TargetNames = new string[] { "UE4Game", "UE4Client", "UE4Server" };
+						foreach(string TargetName in TargetNames)
 						{
-							// Strip the output folder so that this can be used on any machine
-							string RelativeReceiptFileName = ReceiptFileName.MakeRelativeTo(new DirectoryReference(OutputDir));
+							// Need to check for development receipt as we use that for the Engine code in DebugGame
+							UnrealTargetConfiguration EngineConfiguration = (CodeTargetConfiguration == UnrealTargetConfiguration.DebugGame) ? UnrealTargetConfiguration.Development : CodeTargetConfiguration;
+							FileReference ReceiptFileName = TargetReceipt.GetDefaultPath(new DirectoryReference(OutputEnginePath), TargetName, CodeTargetPlatform, EngineConfiguration, Architecture);
 
-							// If we have pre-compiled architectures for this platform then add an entry for each of these -
-							// there isn't a receipt for each architecture like some other platforms
-							if (AllArchNames.Count > 0)
+							if (FileReference.Exists(ReceiptFileName))
 							{
-								foreach (string Arch in AllArchNames)
+								// Strip the output folder so that this can be used on any machine
+								string RelativeReceiptFileName = ReceiptFileName.MakeRelativeTo(new DirectoryReference(OutputDir));
+
+								// If we have pre-compiled architectures for this platform then add an entry for each of these -
+								// there isn't a receipt for each architecture like some other platforms
+								if (AllArchNames.Count > 0)
 								{
-									InstalledConfigs.Add(new InstalledPlatformInfo.InstalledPlatformConfiguration(CodeTargetConfiguration, CodeTargetPlatform, TargetType.Game, Arch, RelativeReceiptFileName, ProjectType, bCanBeDisplayed));
+									foreach (string Arch in AllArchNames)
+									{
+										InstalledConfigs.Add(new InstalledPlatformInfo.InstalledPlatformConfiguration(CodeTargetConfiguration, CodeTargetPlatform, TargetType.Game, Arch, RelativeReceiptFileName, ProjectType, bCanBeDisplayed));
+									}
+								}
+								else
+								{
+									InstalledConfigs.Add(new InstalledPlatformInfo.InstalledPlatformConfiguration(CodeTargetConfiguration, CodeTargetPlatform, TargetType.Game, Architecture, RelativeReceiptFileName, ProjectType, bCanBeDisplayed));
 								}
-							}
-							else
-							{
-								InstalledConfigs.Add(new InstalledPlatformInfo.InstalledPlatformConfiguration(CodeTargetConfiguration, CodeTargetPlatform, TargetType.Game, Architecture, RelativeReceiptFileName, ProjectType, bCanBeDisplayed));
 							}
 						}
 					}
