Fix HoloLens hacks.

#rnx

diff --git a/Engine/Build/InstalledEngineBuild.xml b/Engine/Build/InstalledEngineBuild.xml
--- a/Engine/Build/InstalledEngineBuild.xml
+++ b/Engine/Build/InstalledEngineBuild.xml
@@ -40,7 +40,6 @@
 	<Option Name="WithSwitch" Restrict="true|false" DefaultValue="false" Description="Include the Switch target platform"/>
 	<Option Name="WithLumin" Restrict="true|false" DefaultValue="$(DefaultWithPlatform)" Description="Include the Lumin target platform"/>
 	<Option Name="WithHoloLens" Restrict="true|false" DefaultValue="$(DefaultWithPlatform)" Description="Include the HoloLens target platform"/>
-	
 
 	<!-- Set up target types -->
 	<Option Name="WithClient" Restrict="true|false" DefaultValue="false" Description="Include precompiled client targets"/>
@@ -350,10 +349,11 @@
 					<Compile Target="$(Target)" Platform="Switch" Configuration="$(Configuration)" Tag="#UE4Game Switch" Arguments="-precompile $(TargetDebugInfoArg)" Clean="false"/>
 				</ForEach>
 			</ForEach>
-		</Node>
 
-		<!-- Downstream monolithics and tools HoloLens -->
-		<Node Name="Compile UE4Game HoloLens" Requires="Compile UnrealHeaderTool Win64" Produces="#UE4Game HoloLens;#UE4Game HoloLens Unstripped;#UE4Game HoloLens Stripped;#UE4Game HoloLens Unsigned;#UE4Game HoloLens Signed;#UE4Game HoloLens Includes">
+		</Node>
+	</Agent>
+	<Agent Name="Target HoloLens" Type="CompileWin64;Win64">
+		<Node Name="Compile UE4Game HoloLens" Requires="Compile UnrealHeaderTool Win64" Produces="#UE4Game HoloLens">
 
 			<ForEach Name="Configuration" Values="$(GameConfigurations)">
 				<Compile Target="UE4Game" Platform="HoloLens" Configuration="$(Configuration)" Tag="#UE4Game HoloLens" Arguments="-precompile -allmodules -nolink $(TargetDebugInfoArg)"/>
@@ -572,7 +572,6 @@
 		<Property Name="DDCPlatformsWin64" Value="$(DDCPlatformsWin64)+Switch" If="'$(WithSwitch)' == true"/>
 		<Property Name="DDCPlatformsWin64" Value="$(DDCPlatformsWin64)+Lumin" If="'$(WithLumin)' == true"/>
 		<Property Name="DDCPlatformsWin64" Value="$(DDCPlatformsWin64)+HoloLens" If="'$(WithHoloLens)' == true"/>
-		
 	</Do>
 
 	<Agent Name="DDC Group Win64" Type="Win64">
