Disable HoloLens from binary builds for now. Lots of errors.

#rnx

diff --git a/Engine/Build/InstalledEngineBuild.xml b/Engine/Build/InstalledEngineBuild.xml
--- a/Engine/Build/InstalledEngineBuild.xml
+++ b/Engine/Build/InstalledEngineBuild.xml
@@ -39,7 +39,7 @@
 	<Option Name="WithXboxOne" Restrict="true|false" DefaultValue="false" Description="Include the XboxOne target platform"/>
 	<Option Name="WithSwitch" Restrict="true|false" DefaultValue="false" Description="Include the Switch target platform"/>
 	<Option Name="WithLumin" Restrict="true|false" DefaultValue="$(DefaultWithPlatform)" Description="Include the Lumin target platform"/>
-	<Option Name="WithHoloLens" Restrict="true|false" DefaultValue="$(DefaultWithPlatform)" Description="Include the HoloLens target platform"/>
+	<Option Name="WithHoloLens" Restrict="true|false" DefaultValue="false" Description="Include the HoloLens target platform"/>
 
 	<!-- Set up target types -->
 	<Option Name="WithClient" Restrict="true|false" DefaultValue="false" Description="Include precompiled client targets"/>
