Disable TP_VirtualRealityBP from building DDC, to fix build.

#rnx

diff --git a/Engine/Build/InstalledEngineFilters.xml b/Engine/Build/InstalledEngineFilters.xml
--- a/Engine/Build/InstalledEngineFilters.xml
+++ b/Engine/Build/InstalledEngineFilters.xml
@@ -56,10 +56,11 @@
 		Samples/MobileStarterContent/MobileStarterContent.uproject
 	</Property>
 	<!-- Don't add VR template for Mac and Linux -->
+	<!--
 	<Property Name="ProjectsToBuildDDC" Value="$(ProjectsToBuildDDC)" If="'$(HostPlatform)' != 'Linux'">
 		Templates/TP_VirtualRealityBP/TP_VirtualRealityBP.uproject
 	</Property>
-	
+	-->
 	<!-- Define Editor Filters -->
 	<Property Name="CopyEditorFilter">
 		<!-- This assembly is normally embedded into the UBT executable, but it can technically be rebuilt from an installed build -->
