Fix DDC build failures on Mac due to semicolon in unescaped command line.

#rnx

diff --git a/Engine/Build/InstalledEngineBuild.xml b/Engine/Build/InstalledEngineBuild.xml
--- a/Engine/Build/InstalledEngineBuild.xml
+++ b/Engine/Build/InstalledEngineBuild.xml
@@ -606,7 +606,7 @@
 			<Copy From="$(RootDir)" To="$(DDCDir)" Files="#FilteredCopyList"/>
 			<Copy From="$(CsToolsDir)" To="$(DDCDir)" Files="#Build Tools CS Binaries"/>
 
-			<Command Name="BuildDerivedDataCache" Arguments="-TempDir=&quot;$(DDCDir)&quot; -FeaturePacks=$(ProjectsToBuildDDC) -TargetPlatforms=$(DDCPlatformsWin64) -HostPlatform=Win64 -SavedDir=&quot;$(SavedDir)&quot;"/>
+			<Command Name="BuildDerivedDataCache" Arguments="-TempDir=&quot;$(DDCDir)&quot; -FeaturePacks=&quot;$(ProjectsToBuildDDC)&quot; -TargetPlatforms=$(DDCPlatformsWin64) -HostPlatform=Win64 -SavedDir=&quot;$(SavedDir)&quot;"/>
 
 			<Tag Files="$(SavedDir)/Engine/DerivedDataCache/Compressed.ddp" With="#CompressedDDCWin64"/>
 		</Node>
@@ -653,7 +653,7 @@
 			<Copy From="$(RootDir)" To="$(DDCDir)" Files="#FilteredCopyList"/>
 			<Copy From="$(CsToolsDir)" To="$(DDCDir)" Files="#Build Tools CS Binaries"/>
 
-			<Command Name="BuildDerivedDataCache" Arguments="-TempDir=&quot;$(DDCDir)&quot; -FeaturePacks=$(ProjectsToBuildDDC) -TargetPlatforms=$(DDCPlatformsMac) -HostPlatform=Mac -SavedDir=&quot;$(SavedDir)&quot;"/>
+			<Command Name="BuildDerivedDataCache" Arguments="-TempDir=&quot;$(DDCDir)&quot; -FeaturePacks=&quot;$(ProjectsToBuildDDC)&quot; -TargetPlatforms=$(DDCPlatformsMac) -HostPlatform=Mac -SavedDir=&quot;$(SavedDir)&quot;"/>
 
 			<Tag Files="$(SavedDir)/Engine/DerivedDataCache/Compressed.ddp" With="#CompressedDDCMac"/>
 		</Node>
