Prevent target platforms being built when -set:HostPlatformOnly=true is specified on the command line.

#rnx

diff --git a/Engine/Build/InstalledEngineBuild.xml b/Engine/Build/InstalledEngineBuild.xml
--- a/Engine/Build/InstalledEngineBuild.xml
+++ b/Engine/Build/InstalledEngineBuild.xml
@@ -21,7 +21,7 @@
 	<Property Name="DefaultWithLinux" Value="false" If="$(HostPlatformEditorOnly) Or ($(HostPlatformOnly) And '$(HostPlatform)' != 'Linux')"/>
 
 	<Property Name="DefaultWithPlatform" Value="true"/>
-	<Property Name="DefaultWithPlatform" Value="false" If="$(HostPlatformEditorOnly)"/>
+	<Property Name="DefaultWithPlatform" Value="false" If="$(HostPlatformEditorOnly) Or $(HostPlatformOnly)"/>
 
 	<Property Name="DefaultWithIOS" Value="$(DefaultWithPlatform)"/>
 	<Property Name="DefaultWithIOS" Value="false" If="'$(HostPlatform)' != 'Mac' And !$(AllPlatforms)"/>
