UE-76509 Packaging for Hololens doesn't work in an Installed build
-fixes hololens arm64 packaging error about mspdbcore.dll, this affected non-installed builds as well
-several other jiras are linked from that one that will still remain for packaging from installed builds
-(actually implemented by joe, reviewed and tested by me)

diff --git a/Engine/Build/InstalledEngineBuild.xml b/Engine/Build/InstalledEngineBuild.xml
--- a/Engine/Build/InstalledEngineBuild.xml
+++ b/Engine/Build/InstalledEngineBuild.xml
@@ -356,8 +356,10 @@
 		<Node Name="Compile UE4Game HoloLens" Requires="Compile UnrealHeaderTool Win64" Produces="#UE4Game HoloLens">
 
 			<ForEach Name="Configuration" Values="$(GameConfigurations)">
-				<Compile Target="UE4Game" Platform="HoloLens" Configuration="$(Configuration)" Tag="#UE4Game HoloLens" Arguments="-precompile -allmodules -nolink $(TargetDebugInfoArg)"/>
-				<Compile Target="UE4Game" Platform="HoloLens" Configuration="$(Configuration)" Tag="#UE4Game HoloLens" Arguments="-precompile $(TargetDebugInfoArg)" Clean="false"/>
+				<Compile Target="UE4Game" Platform="HoloLens" Configuration="$(Configuration)" Tag="#UE4Game HoloLens" Arguments="-precompile -allmodules -nolink -architecture=arm64 $(TargetDebugInfoArg)"/>
+				<Compile Target="UE4Game" Platform="HoloLens" Configuration="$(Configuration)" Tag="#UE4Game HoloLens" Arguments="-precompile -architecture=arm64 $(TargetDebugInfoArg)" Clean="false"/>
+				<Compile Target="UE4Game" Platform="HoloLens" Configuration="$(Configuration)" Tag="#UE4Game HoloLens" Arguments="-precompile -allmodules -nolink -architecture=x64 $(TargetDebugInfoArg)"/>
+				<Compile Target="UE4Game" Platform="HoloLens" Configuration="$(Configuration)" Tag="#UE4Game HoloLens" Arguments="-precompile -architecture=x64 $(TargetDebugInfoArg)" Clean="false"/>
 			</ForEach>
 
 		</Node>
diff --git a/Engine/Build/InstalledEngineFilters.xml b/Engine/Build/InstalledEngineFilters.xml
--- a/Engine/Build/InstalledEngineFilters.xml
+++ b/Engine/Build/InstalledEngineFilters.xml
@@ -532,7 +532,20 @@
 	<!-- HoloLens -->
 	<Property Name="CopyHoloLensFilter">
 		<!-- Build Files -->
+		Engine/Binaries/HoloLens/...
 		Engine/Build/HoloLens/...
+		Engine/Intermediate/Deploy/arm64/...
+		Engine/Source/ThirdParty/zlib/...
+		Engine/Source/ThirdParty/IntelTBB/...
+		Engine/Source/ThirdParty/ICU/...
+		Engine/Source/ThirdParty/FreeType2/...
+		Engine/Source/ThirdParty/libPNG/...
+		Engine/Source/ThirdParty/Ogg/...
+		Engine/Source/ThirdParty/Vorbis/...
+		Engine/Source/ThirdParty/libOpus/...
+		Engine/Source/ThirdParty/WindowsMixedRealityInterop/...
+		Engine/Source/ThirdParty/PhysXVehicles/...
+		Engine/Source/ThirdParty/PhysX3/...
 	</Property>
 
 	<Property Name="CopyHoloLensExceptions">
diff --git a/Engine/Source/Programs/UnrealBuildTool/Modes/BuildMode.cs b/Engine/Source/Programs/UnrealBuildTool/Modes/BuildMode.cs
--- a/Engine/Source/Programs/UnrealBuildTool/Modes/BuildMode.cs
+++ b/Engine/Source/Programs/UnrealBuildTool/Modes/BuildMode.cs
@@ -345,13 +345,6 @@ namespace UnrealBuildTool
 						{
 							TargetReceipt Receipt = TargetReceipt.Read(Makefile.ReceiptFile);
 							Log.TraceInformation("Deploying {0} {1} {2}...", Receipt.TargetName, Receipt.Platform, Receipt.Configuration);
-							
-							//@MIXEDREALITY_CHANGE : Begin project dir was missing.
-							if (Receipt.ProjectDir == null)
-							{
-								Receipt.ProjectDir = UnrealBuildTool.EngineDirectory;
-							}
-							//@MIXEDREALITY_CHANGE : End
 
 							UEBuildPlatform.GetBuildPlatform(Receipt.Platform).Deploy(Receipt);
 						}
diff --git a/Engine/Source/Programs/UnrealBuildTool/Platform/HoloLens/HoloLensToolChain.cs b/Engine/Source/Programs/UnrealBuildTool/Platform/HoloLens/HoloLensToolChain.cs
--- a/Engine/Source/Programs/UnrealBuildTool/Platform/HoloLens/HoloLensToolChain.cs
+++ b/Engine/Source/Programs/UnrealBuildTool/Platform/HoloLens/HoloLensToolChain.cs
@@ -48,6 +48,10 @@ namespace UnrealBuildTool
 				{
 				}
 			}
+
+			// by default tools chains don't parse arguments, but we want to be able to check the -architectures flag defined above. This is
+			// only necessary when AndroidToolChain is used during UAT
+			CommandLine.ParseArguments(Environment.GetCommandLineArgs(), this);
 		}
 
 		void AppendCLArguments_Global(CppCompileEnvironment CompileEnvironment, VCEnvironment EnvVars, List<string> Arguments)
diff --git a/Engine/Source/Programs/UnrealBuildTool/Platform/HoloLens/UEBuildHoloLens.cs b/Engine/Source/Programs/UnrealBuildTool/Platform/HoloLens/UEBuildHoloLens.cs
--- a/Engine/Source/Programs/UnrealBuildTool/Platform/HoloLens/UEBuildHoloLens.cs
+++ b/Engine/Source/Programs/UnrealBuildTool/Platform/HoloLens/UEBuildHoloLens.cs
@@ -28,6 +28,8 @@ namespace UnrealBuildTool
 		/// <summary>
 		/// Architecture of Target.
 		/// </summary>
+		[CommandLine("x64", Value = "x64")]
+		[CommandLine("arm64", Value = "ARM64")]
 		public WindowsArchitecture Architecture = WindowsArchitecture.x64;
 
 		/// <summary>
diff --git a/Engine/Source/Programs/UnrealBuildTool/Platform/Windows/VCEnvironment.cs b/Engine/Source/Programs/UnrealBuildTool/Platform/Windows/VCEnvironment.cs
--- a/Engine/Source/Programs/UnrealBuildTool/Platform/Windows/VCEnvironment.cs
+++ b/Engine/Source/Programs/UnrealBuildTool/Platform/Windows/VCEnvironment.cs
@@ -154,6 +154,12 @@ namespace UnrealBuildTool
 				AddDirectoryToPath(GetVCToolPath(ToolChain, ToolChainDir, WindowsArchitecture.x86));
 				AddDirectoryToPath(GetVCToolPath(ToolChain, ToolChainDir, WindowsArchitecture.x64));
 			}
+			else if (Architecture == WindowsArchitecture.ARM64)
+			{
+				AddDirectoryToPath(GetVCToolPath(ToolChain, ToolChainDir, WindowsArchitecture.ARM64));
+				AddDirectoryToPath(GetVCToolPath(ToolChain, ToolChainDir, WindowsArchitecture.x86));
+				AddDirectoryToPath(GetVCToolPath(ToolChain, ToolChainDir, WindowsArchitecture.x64));
+			}
 
 
 			// Get all the system include paths
