Recreate 3958825 so that the maya live link plugin gets added to engine extras again

diff --git a/Engine/Build/InstalledEngineFilters.xml b/Engine/Build/InstalledEngineFilters.xml
--- a/Engine/Build/InstalledEngineFilters.xml
+++ b/Engine/Build/InstalledEngineFilters.xml
@@ -219,6 +219,7 @@
 			<!-- Extra tools -->
 			Engine/Extras/3dsMaxScripts/...
 			Engine/Extras/VisualStudioDebugging/UE4.natvis
+			Engine/Extras/MayaLiveLink/...
 			Engine/Extras/MayaVelocityGridExporter/...
 			Engine/Extras/UnrealVS/...
 			Engine/Extras/Redist/en-us/*
