Restore missing Win32 executables from binary releases.

#rnx

diff --git a/Engine/Build/InstalledEngineBuild.xml b/Engine/Build/InstalledEngineBuild.xml
--- a/Engine/Build/InstalledEngineBuild.xml
+++ b/Engine/Build/InstalledEngineBuild.xml
@@ -756,7 +756,7 @@
 				<Property Name="CopyInstalledExceptions" Value="$(CopyInstalledExceptions);$(CopyWin64Exceptions)"/>
 			</Do>
 			<Do If="'$(WithWin32)' == true">
-				<Tag Files="#UE4Game Win32" Except="#UE4Game Win32 Unstripped;#UE4Game Win32 Unsigned" With="#Installed Win64"/>
+				<Tag Files="#UE4Game Win32" Except="#UE4Game Win32 Unstripped;#UE4Game Win32 Unsigned;#Build Tools Win32" With="#Installed Win64"/>
 				<Tag Files="#UE4Game Win32 Stripped;#UE4Game Win32 Signed" With="#Saved Output"/>
 				<Property Name="CopyInstalledFilter" Value="$(CopyInstalledFilter);$(CopyWin32Filter)"/>
 				<Property Name="CopyInstalledExceptions" Value="$(CopyInstalledExceptions);$(CopyWin32Exceptions)"/>
@@ -831,10 +831,12 @@
 
 			<!-- Tag Win64 files that need to be stripped -->
 			<Tag Files="#Compile UnrealHeaderTool Win64;#Build Tools Win64" Filter="$(Win64StripFilter)" Except="$(Win64StripExceptions);$(CopyInstalledExceptions)" With="#Strip Tools for Installed Win64"/>
+			<Tag Files="#Build Tools Win32" Filter="$(Win32StripFilter)" Except="$(CopyInstalledExceptions)" With="#Strip Tools for Installed Win64" If="$(WithWin32)"/>
 
 			<Do If="'$(SignExecutables)' == true">
 				<!-- Tag files that need to be signed -->
 				<Tag Files="#Compile UnrealHeaderTool Win64;#Build Tools Win64" Filter="$(WindowsSignFilter)" Except="$(CopyInstalledExceptions);$(SignExceptions)" With="#Sign Tools for Installed Win64"/>
+				<Tag Files="#Build Tools Win32" Filter="$(WindowsSignFilter)" Except="$(CopyInstalledExceptions);$(SignExceptions)" With="#Sign Tools for Installed Win64" If="$(WithWin32)"/>
 			</Do>
 
 			<!-- Tag files that need to be copied directly -->
