Attempted workaround for sample editor build errors; build the development editor after the debuggame editor, to ensure the saved action graph reflects the development editor.

diff --git a/Engine/Build/InstalledEngineBuild.xml b/Engine/Build/InstalledEngineBuild.xml
--- a/Engine/Build/InstalledEngineBuild.xml
+++ b/Engine/Build/InstalledEngineBuild.xml
@@ -119,9 +119,8 @@
 			<Compile Target="UnrealHeaderTool" Platform="Win64" Configuration="Development" Arguments="-precompile"/>
 		</Node>
 		<Node Name="Compile UE4Editor Win64" Requires="Compile UnrealHeaderTool Win64" Produces="#UE4Editor Win64;#UE4Editor Win64 Unstripped;#UE4Editor Win64 Stripped;#UE4Editor Win64 Unsigned;#UE4Editor Win64 Signed">
-			<Compile Target="UE4Editor" Platform="Win64" Configuration="Development" Tag="#UE4Editor Win64" Arguments="-precompile -allmodules"/>
 			<Compile Target="UE4Editor" Platform="Win64" Configuration="DebugGame" Tag="#UE4Editor Win64" Arguments="-precompile -allmodules"/>
-			<Compile Target="UE4Editor" Platform="Win64" Configuration="DebugGame" Tag="#UE4Editor Win64" Arguments="-precompile"/>
+			<Compile Target="UE4Editor" Platform="Win64" Configuration="Development" Tag="#UE4Editor Win64" Arguments="-precompile -allmodules"/>
 
 			<Do If="$(EmbedSrcSrvInfo)">
 				<!-- Embed source info into the PDB files. Should be done from this machine to ensure that paths are correct. -->
@@ -159,9 +158,8 @@
 			<Compile Target="UnrealHeaderTool" Platform="Mac" Configuration="Development" Arguments="-precompile -allmodules" If="'$(HostPlatform)' != 'Win64'"/>
 		</Node>
 		<Node Name="Compile UE4Editor Mac" Requires="Compile UnrealHeaderTool Mac" Produces="#UE4Editor Mac;#UE4Editor Mac Unstripped;#UE4Editor Mac Stripped;#UE4Editor Mac Unsigned;#UE4Editor Mac Signed">
-			<Compile Target="UE4Editor" Platform="Mac" Configuration="Development" Tag="#UE4Editor Mac" Arguments="-precompile -allmodules"/>
 			<Compile Target="UE4Editor" Platform="Mac" Configuration="DebugGame" Tag="#UE4Editor Mac" Arguments="-precompile -allmodules"/>
-			<Compile Target="UE4Editor" Platform="Mac" Configuration="DebugGame" Tag="#UE4Editor Mac" Arguments="-precompile"/>
+			<Compile Target="UE4Editor" Platform="Mac" Configuration="Development" Tag="#UE4Editor Mac" Arguments="-precompile -allmodules"/>
 
 			<Tag Files="#UE4Editor Mac" Filter="$(MacStripFilter)" With="#UE4Editor Mac Unstripped"/>
 
@@ -190,9 +188,8 @@
 			<Compile Target="UnrealHeaderTool" Platform="Linux" Configuration="Development" Arguments="-precompile -allmodules"/>
 		</Node>
 		<Node Name="Compile UE4Editor Linux" Requires="Compile UnrealHeaderTool $(HostPlatform)" Produces="#UE4Editor Linux;#UE4Editor Linux Unstripped;#UE4Editor Linux Stripped">
-			<Compile Target="UE4Editor" Platform="Linux" Configuration="Development" Tag="#UE4Editor Linux" Arguments="-precompile -allmodules"/>
 			<Compile Target="UE4Editor" Platform="Linux" Configuration="DebugGame" Tag="#UE4Editor Linux" Arguments="-precompile -allmodules"/>
-			<Compile Target="UE4Editor" Platform="Linux" Configuration="DebugGame" Tag="#UE4Editor Linux" Arguments="-precompile"/>
+			<Compile Target="UE4Editor" Platform="Linux" Configuration="Development" Tag="#UE4Editor Linux" Arguments="-precompile -allmodules"/>
 
 			<Tag Files="#UE4Editor Linux" Filter="$(LinuxStripFilter)" With="#UE4Editor Linux Unstripped"/>
 
