Remove logic that prevents UHT compiling on Mac during a remote build.

diff --git a/Engine/Build/InstalledEngineBuild.xml b/Engine/Build/InstalledEngineBuild.xml
--- a/Engine/Build/InstalledEngineBuild.xml
+++ b/Engine/Build/InstalledEngineBuild.xml
@@ -155,7 +155,7 @@
 	<!-- Base editor and tools for compile on Mac -->
 	<Agent Name="Editor Mac" Type="CompileMac;Mac">
 		<Node Name="Compile UnrealHeaderTool Mac" Requires="Update Version Files">
-			<Compile Target="UnrealHeaderTool" Platform="Mac" Configuration="Development" Arguments="-precompile -allmodules" If="'$(HostPlatform)' != 'Win64'"/>
+			<Compile Target="UnrealHeaderTool" Platform="Mac" Configuration="Development" Arguments="-precompile -allmodules"/>
 		</Node>
 		<Node Name="Compile UE4Editor Mac" Requires="Compile UnrealHeaderTool Mac" Produces="#UE4Editor Mac;#UE4Editor Mac Unstripped;#UE4Editor Mac Stripped;#UE4Editor Mac Unsigned;#UE4Editor Mac Signed">
 			<Compile Target="UE4Editor" Platform="Mac" Configuration="DebugGame" Tag="#UE4Editor Mac" Arguments="-precompile -allmodules"/>
