Add Rsync filters to installed build on Windows.

diff --git a/Engine/Build/InstalledEngineFilters.xml b/Engine/Build/InstalledEngineFilters.xml
--- a/Engine/Build/InstalledEngineFilters.xml
+++ b/Engine/Build/InstalledEngineFilters.xml
@@ -352,7 +352,6 @@
 		<!-- Build Files -->
 		Engine/Build/IOS/...
 	</Property>
-		<!-- Win64 specific -->
 	<Do If="'$(HostPlatform)' == 'Win64'">
 		<Property Name="CopyIOSFilter" Value="$(CopyIOSFilter)">
 			Engine/Binaries/DotNET/IOS/openssl.exe
@@ -360,6 +359,7 @@
 			Engine/Binaries/ThirdParty/ICU/icu4c-53_1/Mac/...
 			Engine/Binaries/ThirdParty/IOS/*
 			Engine/Build/BatchFiles/MakeAndInstallSSHKey.bat
+			Engine/Build/Rsync/...
 			Engine/Extras/ThirdPartyNotUE/DeltaCopy/...
 		</Property>
 	</Do>
