Fix IPhonePackager not being included in installed builds created on a Mac.

diff --git a/Engine/Build/InstalledEngineBuild.xml b/Engine/Build/InstalledEngineBuild.xml
--- a/Engine/Build/InstalledEngineBuild.xml
+++ b/Engine/Build/InstalledEngineBuild.xml
@@ -87,7 +87,7 @@
 	<Property Name="LocalInstalledDirLinux" Value="$(RootDir)/LocalBuilds/Engine/Linux"/>
 
 	<!-- Directory for storing build products like the Compressed DDC, Stripped and Signed files -->
-	<Property Name="SavedOutput" Value="$(RootDir)/Saved"/>
+	<Property Name="SavedOutput" Value="$(RootDir)/Engine/Saved"/>
 
 	<!-- Directory for storing a stripped version of UAT for installed builds -->
 	<Property Name="CsToolsDir" Value="$(SavedOutput)/CsTools"/>
@@ -465,10 +465,16 @@
 				<CsCompile Project="$(CsToolsDir)/Engine/Source/Programs/nDisplayListener/nDisplayListener.csproj" Configuration="Development" Platform="AnyCPU" Tag="#Build Tools CS Binaries;#Sign Binaries" TagReferences="#Build Tools CS Binaries"/>
 				<CsCompile Project="$(CsToolsDir)/Engine/Source/Programs/PS4/PS4DevKitUtil/PS4DevKitUtil.csproj" Configuration="Development" Platform="AnyCPU" Tag="#Build Tools CS Binaries;#Sign Binaries" TagReferences="#Build Tools CS Binaries" If="'$(WithPS4)' == true"/>
 				<CsCompile Project="$(CsToolsDir)/Engine/Source/Programs/PS4/PS4SymbolTool/PS4SymbolTool.csproj" Configuration="Release" Platform="AnyCPU" Tag="#Build Tools CS Binaries;#Sign Binaries" TagReferences="#Build Tools CS Binaries" If="'$(WithPS4)' == true"/>
-				<!-- TODO: Change these projects so that they can be compiled on non-windows platforms (currently use WinForms) -->
 				<CsCompile Project="$(CsToolsDir)/Engine/Source/Programs/UnrealControls/UnrealControls.csproj" Configuration="Development" Platform="AnyCPU" Tag="#Build Tools CS Binaries;#Sign Binaries" TagReferences="#Build Tools CS Binaries"/>
 				<CsCompile Project="$(CsToolsDir)/Engine/Source/Programs/IOS/iPhonePackager/iPhonePackager.csproj" Configuration="Development" Platform="AnyCPU" Arguments="/verbosity:minimal /target:Rebuild" Tag="#Build Tools CS Binaries;#Sign Binaries" TagReferences="#Build Tools CS Binaries"/>
 			</Do>
+			
+			<!-- Compile IPP. Mac/Linux bundled versions of Mono do not have Winforms, so we just use the existing one. -->
+			<Property Name="CompileIPP" Value="true"/>
+			<Property Name="CompileIPP" Value="false" If="'$(HostPlatform)' != 'Win64'"/>
+			<CsCompile Project="$(CsToolsDir)/Engine/Source/Programs/UnrealControls/UnrealControls.csproj" Configuration="Development" Platform="AnyCPU" Tag="#Build Tools CS Binaries;#Sign Binaries" TagReferences="#Build Tools CS Binaries" EnumerateOnly="!$(CompileIPP)"/>
+			<CsCompile Project="$(CsToolsDir)/Engine/Source/Programs/IOS/iPhonePackager/iPhonePackager.csproj" Configuration="Development" Platform="AnyCPU" Arguments="/verbosity:minimal /target:Rebuild" Tag="#Build Tools CS Binaries;#Sign Binaries" TagReferences="#Build Tools CS Binaries" EnumerateOnly="!$(CompileIPP)"/>
+			
 			<CsCompile Project="$(CsToolsDir)/Engine/Source/Programs/NetworkProfiler/NetworkProfiler/NetworkProfiler.csproj" Configuration="Development" Platform="AnyCPU" Tag="#Build Tools CS Binaries" TagReferences="#Build Tools CS Binaries;#Sign Binaries" If="'$(HostPlatform)' != 'Linux'"/>
 			<CsCompile Project="$(CsToolsDir)/Engine/Source/Programs/IOS/DeploymentServer/DeploymentServer.csproj" Configuration="Development" Platform="AnyCPU" Arguments="/verbosity:minimal /target:Rebuild" Tag="#Build Tools CS Binaries;#Sign Binaries" TagReferences="#Build Tools CS Binaries"/>
 			<CsCompile Project="$(CsToolsDir)/Engine/Source/Programs/IOS/MobileDeviceInterface/MobileDeviceInterface.csproj" Configuration="Development" Platform="AnyCPU" Arguments="/verbosity:minimal /target:Rebuild" Tag="#Build Tools CS Binaries;#Sign Binaries" TagReferences="#Build Tools CS Binaries"/>
