Compile DebugGame version of the launch module for installed builds.

diff --git a/Engine/Build/InstalledEngineBuild.xml b/Engine/Build/InstalledEngineBuild.xml
--- a/Engine/Build/InstalledEngineBuild.xml
+++ b/Engine/Build/InstalledEngineBuild.xml
@@ -41,7 +41,7 @@
 	<Option Name="EmbedSrcSrvInfo" Restrict="true|false" DefaultValue="false" Description="Whether to add Source indexing to Windows game apps so they can be added to a symbol server"/>
 
 	<!-- Which game configurations to include for packaged applications -->
-	<Option Name="GameConfigurations" DefaultValue="Development;Shipping" Description="Which game configurations to include for packaged applications"/>
+	<Option Name="GameConfigurations" DefaultValue="DebugGame;Development;Shipping" Description="Which game configurations to include for packaged applications"/>
 
 	<!-- Generate full debug info for binary editor and packaged application builds -->
 	<Option Name="WithFullDebugInfo" Restrict="true|false" DefaultValue="false" Description="Generate full debug info for binary editor and packaged application builds"/>
