Make sure Newtonsoft is included in the installed engine build. (Also ignore *.build.cs that aren't enumerated from the targets we build; they won't be usable)

diff --git a/Engine/Build/InstalledEngineFilters.xml b/Engine/Build/InstalledEngineFilters.xml
--- a/Engine/Build/InstalledEngineFilters.xml
+++ b/Engine/Build/InstalledEngineFilters.xml
@@ -61,6 +61,7 @@
 	<!-- Define Editor Filters -->
 	<Property Name="CopyEditorFilter">
 		<!-- This assembly is normally embedded into the UBT executable, but it can technically be rebuilt from an installed build -->
+		Engine/Binaries/ThirdParty/Newtonsoft/...
 		Engine/Binaries/ThirdParty/VisualStudio/...
 
 		<!-- In-editor documentation -->
@@ -91,7 +92,6 @@
 		Engine/Source/Runtime/...
 		Engine/Source/Developer/...
 		Engine/Source/Editor/...
-		Engine/Source/ThirdParty/.../*.Build.cs
 		Engine/Source/ThirdParty/Licenses/...
 		Engine/Source/Programs/AutomationTool/...
 		Engine/Source/Programs/AutomationToolLauncher/...
