Include Mac build files on Windows (eg. shell scripts, mono) to support remote compilation.

diff --git a/Engine/Build/InstalledEngineFilters.xml b/Engine/Build/InstalledEngineFilters.xml
--- a/Engine/Build/InstalledEngineFilters.xml
+++ b/Engine/Build/InstalledEngineFilters.xml
@@ -220,24 +220,27 @@
 		</Property>
 	</Do>
 
+	<!-- Files needed to build on Mac, either natively or via the remote toolchain -->
+	<Property Name="MacBuildFiles">
+		<!-- Build batch files -->
+		Engine/Build/BatchFiles/Mac/Build.sh
+		Engine/Build/BatchFiles/Mac/GenerateLLDBInit.sh
+		Engine/Build/BatchFiles/Mac/FixDependencyFiles.sh
+		Engine/Build/BatchFiles/Mac/GenerateProjectFiles.sh
+		Engine/Build/BatchFiles/Mac/SetupMono.sh
+		Engine/Build/BatchFiles/Mac/FixMonoFiles.sh
+		Engine/Build/BatchFiles/Mac/RunMono.sh
+		Engine/Build/BatchFiles/Mac/RunXBuild.sh
+		Engine/Build/BatchFiles/RunUAT.command
+		Engine/Build/BatchFiles/RunUAT.sh
+
+		<!-- Mono -->
+		Engine/Binaries/ThirdParty/Mono/Mac/...
+	</Property>
+
 	<!-- Mac-specific dependencies -->
 	<Do If="'$(HostPlatform)' == 'Mac'">
-		<Property Name="CopyEditorFilter" Value="$(CopyEditorFilter)">
-			<!-- Build batch files -->
-			Engine/Build/BatchFiles/Mac/Build.sh
-			Engine/Build/BatchFiles/Mac/GenerateLLDBInit.sh
-			Engine/Build/BatchFiles/Mac/FixDependencyFiles.sh
-			Engine/Build/BatchFiles/Mac/GenerateProjectFiles.sh
-			Engine/Build/BatchFiles/Mac/SetupMono.sh
-			Engine/Build/BatchFiles/Mac/FixMonoFiles.sh
-			Engine/Build/BatchFiles/Mac/RunMono.sh
-			Engine/Build/BatchFiles/Mac/RunXBuild.sh
-			Engine/Build/BatchFiles/RunUAT.command
-			Engine/Build/BatchFiles/RunUAT.sh
-
-			<!-- Mono -->
-			Engine/Binaries/ThirdParty/Mono/Mac/...
-
+		<Property Name="CopyEditorFilter" Value="$(CopyEditorFilter);$(MacBuildFiles)">
 			<!-- Extra tools -->
 			Engine/Extras/LLDBDataFormatters/UE4DataFormatters.py
 		</Property>
@@ -348,7 +351,7 @@
 	</Property>
 
 	<!-- IOS -->
-	<Property Name="CopyIOSFilter">
+	<Property Name="CopyIOSFilter" Value="$(MacBuildFiles)">
 		<!-- Build Files -->
 		Engine/Build/IOS/...
 	</Property>
@@ -372,7 +375,7 @@
 	</Property>
 
 	<!-- TVOS -->
-	<Property Name="CopyTVOSFilter">
+	<Property Name="CopyTVOSFilter" Value="$(MacBuildFiles)">
 		<!-- Build Files -->
 		Engine/Build/TVOS/...
 	</Property>
