Fix debug info being included for Lumin binaries on Mac.

diff --git a/Engine/Build/InstalledEngineBuild.xml b/Engine/Build/InstalledEngineBuild.xml
--- a/Engine/Build/InstalledEngineBuild.xml
+++ b/Engine/Build/InstalledEngineBuild.xml
@@ -409,8 +409,8 @@
 		<Node Name="Compile UE4Game LuminMac" Requires="Compile UnrealHeaderTool Mac" Produces="#UE4Game LuminMac;#UE4Game LuminMac Unstripped;#UE4Game LuminMac Stripped">
 
 			<ForEach Name="Configuration" Values="$(GameConfigurations)">
-				<Compile Target="UE4Game" Platform="Lumin" Configuration="$(Configuration)" Tag="#UE4Game LuminMac" Arguments="-precompile -allmodules -nolink -architectures=$(AndroidArchitectures) -gpuarchitectures=$(LuminGPUArchitectures)" AllowParallelExecutor="false"/>
-				<Compile Target="UE4Game" Platform="Lumin" Configuration="$(Configuration)" Tag="#UE4Game LuminMac" Arguments="-precompile -architectures=$(AndroidArchitectures) -gpuarchitectures=$(LuminGPUArchitectures)" AllowParallelExecutor="false" Clean="false"/>
+				<Compile Target="UE4Game" Platform="Lumin" Configuration="$(Configuration)" Tag="#UE4Game LuminMac" Arguments="-precompile -allmodules -nolink -architectures=$(AndroidArchitectures) -gpuarchitectures=$(LuminGPUArchitectures) $(TargetDebugInfoArg)" AllowParallelExecutor="false"/>
+				<Compile Target="UE4Game" Platform="Lumin" Configuration="$(Configuration)" Tag="#UE4Game LuminMac" Arguments="-precompile -architectures=$(AndroidArchitectures) -gpuarchitectures=$(LuminGPUArchitectures) $(TargetDebugInfoArg)" AllowParallelExecutor="false" Clean="false"/>
 			</ForEach>
 
 			<Tag Files="#UE4Game LuminMac" Filter="$(LuminStripFilter)" With="#UE4Game LuminMac Unstripped"/>
