Copy Windows DDCUtils DLLs when building DDC for Linux on the Win64 host platform.

diff --git a/Engine/Build/InstalledEngineBuild.xml b/Engine/Build/InstalledEngineBuild.xml
--- a/Engine/Build/InstalledEngineBuild.xml
+++ b/Engine/Build/InstalledEngineBuild.xml
@@ -647,8 +647,10 @@
 			<!-- Additional required files that would be removed by confidential filter -->
 			<Tag Files="Engine/Build/NotForLicensees/EpicInternal.txt" With="#FilteredCopyList" If="Exists('Engine/Build/NotForLicensees/EpicInternal.txt')"/>
 			<Tag Files="Engine/Binaries/Linux/NotForLicensees/ShaderCompileWorker.modules" With="#FilteredCopyList" If="Exists('Engine/Build/NotForLicensees/EpicInternal.txt')"/>
+			<Tag Files="Engine/Binaries/Linux/NotForLicensees/ShaderCompileWorker-DDCUtils.dll" With="#FilteredCopyList" If="Exists('Engine/Build/NotForLicensees/EpicInternal.txt')"/>
 			<Tag Files="Engine/Binaries/Linux/NotForLicensees/libShaderCompileWorker-DDCUtils.so" With="#FilteredCopyList" If="Exists('Engine/Build/NotForLicensees/EpicInternal.txt')"/>
 			<Tag Files="Engine/Binaries/Linux/NotForLicensees/UE4Editor.modules" With="#FilteredCopyList" If="Exists('Engine/Build/NotForLicensees/EpicInternal.txt')"/>
+			<Tag Files="Engine/Binaries/Linux/NotForLicensees/UE4Editor-DDCUtils.dll" With="#FilteredCopyList" If="Exists('Engine/Build/NotForLicensees/EpicInternal.txt')"/>
 			<Tag Files="Engine/Binaries/Linux/NotForLicensees/libUE4Editor-DDCUtils.so" With="#FilteredCopyList" If="Exists('Engine/Build/NotForLicensees/EpicInternal.txt')"/>
 			<!-- Copy everything to a temporary directory -->
 			<Copy From="$(RootDir)" To="$(DDCDir)" Files="#FilteredCopyList"/>
