Fix path to Win64 DLLs when building DDC on Linux.

diff --git a/Engine/Build/InstalledEngineBuild.xml b/Engine/Build/InstalledEngineBuild.xml
--- a/Engine/Build/InstalledEngineBuild.xml
+++ b/Engine/Build/InstalledEngineBuild.xml
@@ -647,10 +647,10 @@
 			<!-- Additional required files that would be removed by confidential filter -->
 			<Tag Files="Engine/Build/NotForLicensees/EpicInternal.txt" With="#FilteredCopyList" If="Exists('Engine/Build/NotForLicensees/EpicInternal.txt')"/>
 			<Tag Files="Engine/Binaries/Linux/NotForLicensees/ShaderCompileWorker.modules" With="#FilteredCopyList" If="Exists('Engine/Build/NotForLicensees/EpicInternal.txt')"/>
-			<Tag Files="Engine/Binaries/Linux/NotForLicensees/ShaderCompileWorker-DDCUtils.dll" With="#FilteredCopyList" If="Exists('Engine/Build/NotForLicensees/EpicInternal.txt')"/>
+			<Tag Files="Engine/Binaries/Win64/NotForLicensees/ShaderCompileWorker-DDCUtils.dll" With="#FilteredCopyList" If="Exists('Engine/Build/NotForLicensees/EpicInternal.txt') And '$(HostPlatform)' == 'Win64'"/>
 			<Tag Files="Engine/Binaries/Linux/NotForLicensees/libShaderCompileWorker-DDCUtils.so" With="#FilteredCopyList" If="Exists('Engine/Build/NotForLicensees/EpicInternal.txt')"/>
 			<Tag Files="Engine/Binaries/Linux/NotForLicensees/UE4Editor.modules" With="#FilteredCopyList" If="Exists('Engine/Build/NotForLicensees/EpicInternal.txt')"/>
-			<Tag Files="Engine/Binaries/Linux/NotForLicensees/UE4Editor-DDCUtils.dll" With="#FilteredCopyList" If="Exists('Engine/Build/NotForLicensees/EpicInternal.txt')"/>
+			<Tag Files="Engine/Binaries/Win64/NotForLicensees/UE4Editor-DDCUtils.dll" With="#FilteredCopyList" If="Exists('Engine/Build/NotForLicensees/EpicInternal.txt') And '$(HostPlatform)' == 'Win64'"/>
 			<Tag Files="Engine/Binaries/Linux/NotForLicensees/libUE4Editor-DDCUtils.so" With="#FilteredCopyList" If="Exists('Engine/Build/NotForLicensees/EpicInternal.txt')"/>
 			<!-- Copy everything to a temporary directory -->
 			<Copy From="$(RootDir)" To="$(DDCDir)" Files="#FilteredCopyList"/>
