Fix warnings about DDCUtils when building Linux DDC.

diff --git a/Engine/Build/InstalledEngineBuild.xml b/Engine/Build/InstalledEngineBuild.xml
--- a/Engine/Build/InstalledEngineBuild.xml
+++ b/Engine/Build/InstalledEngineBuild.xml
@@ -646,12 +646,16 @@
 			<Tag Files="#ToCopy" Except="$(ConfidentialExceptions);.../Source/...;.../Intermediate/..." With="#FilteredCopyList"/>
 			<!-- Additional required files that would be removed by confidential filter -->
 			<Tag Files="Engine/Build/NotForLicensees/EpicInternal.txt" With="#FilteredCopyList" If="Exists('Engine/Build/NotForLicensees/EpicInternal.txt')"/>
+			<!-- When running on Linux -->
 			<Tag Files="Engine/Binaries/Linux/NotForLicensees/ShaderCompileWorker.modules" With="#FilteredCopyList" If="Exists('Engine/Build/NotForLicensees/EpicInternal.txt')"/>
-			<Tag Files="Engine/Binaries/Win64/NotForLicensees/ShaderCompileWorker-DDCUtils.dll" With="#FilteredCopyList" If="Exists('Engine/Build/NotForLicensees/EpicInternal.txt') And '$(HostPlatform)' == 'Win64'"/>
 			<Tag Files="Engine/Binaries/Linux/NotForLicensees/libShaderCompileWorker-DDCUtils.so" With="#FilteredCopyList" If="Exists('Engine/Build/NotForLicensees/EpicInternal.txt')"/>
 			<Tag Files="Engine/Binaries/Linux/NotForLicensees/UE4Editor.modules" With="#FilteredCopyList" If="Exists('Engine/Build/NotForLicensees/EpicInternal.txt')"/>
-			<Tag Files="Engine/Binaries/Win64/NotForLicensees/UE4Editor-DDCUtils.dll" With="#FilteredCopyList" If="Exists('Engine/Build/NotForLicensees/EpicInternal.txt') And '$(HostPlatform)' == 'Win64'"/>
 			<Tag Files="Engine/Binaries/Linux/NotForLicensees/libUE4Editor-DDCUtils.so" With="#FilteredCopyList" If="Exists('Engine/Build/NotForLicensees/EpicInternal.txt')"/>
+			<!-- When running on Windows -->
+			<Tag Files="Engine/Binaries/Win64/NotForLicensees/ShaderCompileWorker.modules" With="#FilteredCopyList" If="Exists('Engine/Build/NotForLicensees/EpicInternal.txt') And '$(HostPlatform)' == 'Win64'"/>
+			<Tag Files="Engine/Binaries/Win64/NotForLicensees/ShaderCompileWorker-DDCUtils.dll" With="#FilteredCopyList" If="Exists('Engine/Build/NotForLicensees/EpicInternal.txt') And '$(HostPlatform)' == 'Win64'"/>
+			<Tag Files="Engine/Binaries/Win64/NotForLicensees/UE4Editor.modules" With="#FilteredCopyList" If="Exists('Engine/Build/NotForLicensees/EpicInternal.txt') And '$(HostPlatform)' == 'Win64'"/>
+			<Tag Files="Engine/Binaries/Win64/NotForLicensees/UE4Editor-DDCUtils.dll" With="#FilteredCopyList" If="Exists('Engine/Build/NotForLicensees/EpicInternal.txt') And '$(HostPlatform)' == 'Win64'"/>
 			<!-- Copy everything to a temporary directory -->
 			<Copy From="$(RootDir)" To="$(DDCDir)" Files="#FilteredCopyList"/>
 			<Copy From="$(CsToolsDir)" To="$(DDCDir)" Files="#Build Tools CS Binaries"/>
