Merging fixes to warnings when building DDC on Linux.

diff --git a/Engine/Build/InstalledEngineBuild.xml b/Engine/Build/InstalledEngineBuild.xml
--- a/Engine/Build/InstalledEngineBuild.xml
+++ b/Engine/Build/InstalledEngineBuild.xml
@@ -637,7 +637,7 @@
 			<TagReceipt Files="#TargetReceipts" RuntimeDependencies="true" With="#ToCopy"/>
 			<Tag Files="Engine/Content/..." Except="*.psd;*.pdn;*.fbx;*.po" With="#ToCopy"/>
 			<Tag Files="Engine/Config/..." Except="*.vdf" With="#ToCopy"/>
-			<Tag Files="Engine/Plugins/....uplugin;Engine/Plugins/.../Content/...;Engine/Plugins/.../Resources/...;Engine/Plugins/.../Shaders/...;Engine/Plugins/.../Templates/..." Except="Engine/Plugins/Runtime/TwitchLiveStreaming/...;Engine/Plugins/.../Binaries/Win64/...;Engine/Plugins/.../Binaries/Win32/...;Engine/Plugins/.../Binaries/Mac/..." With="#ToCopy"/>
+			<Tag Files="Engine/Plugins/....uplugin;Engine/Plugins/.../Config/...;Engine/Plugins/.../Content/...;Engine/Plugins/.../Resources/...;Engine/Plugins/.../Shaders/...;Engine/Plugins/.../Templates/..." Except="Engine/Plugins/Runtime/TwitchLiveStreaming/...;Engine/Plugins/.../Binaries/Win32/...;Engine/Plugins/.../Binaries/Mac/..." With="#ToCopy"/>
 			<Tag Files="Engine/Shaders/..." With="#ToCopy"/>
 			<Tag Files="Engine/Binaries/DotNET/Ionic.Zip.Reduced.dll" With="#ToCopy"/>
 			<Tag Files="Engine/Binaries/DotNET/OneSky.dll" With="#ToCopy"/>
@@ -646,12 +646,16 @@
 			<Tag Files="#ToCopy" Except="$(ConfidentialExceptions);.../Source/...;.../Intermediate/..." With="#FilteredCopyList"/>
 			<!-- Additional required files that would be removed by confidential filter -->
 			<Tag Files="Engine/Build/NotForLicensees/EpicInternal.txt" With="#FilteredCopyList" If="Exists('Engine/Build/NotForLicensees/EpicInternal.txt')"/>
+			<!-- When running on Linux -->
 			<Tag Files="Engine/Binaries/Linux/NotForLicensees/ShaderCompileWorker.modules" With="#FilteredCopyList" If="Exists('Engine/Build/NotForLicensees/EpicInternal.txt')"/>
-			<Tag Files="Engine/Binaries/Win64/NotForLicensees/ShaderCompileWorker-DDCUtils.dll" With="#FilteredCopyList" If="Exists('Engine/Build/NotForLicensees/EpicInternal.txt') And '$(HostPlatform)' == 'Win64'"/>
 			<Tag Files="Engine/Binaries/Linux/NotForLicensees/libShaderCompileWorker-DDCUtils.so" With="#FilteredCopyList" If="Exists('Engine/Build/NotForLicensees/EpicInternal.txt')"/>
 			<Tag Files="Engine/Binaries/Linux/NotForLicensees/UE4Editor.modules" With="#FilteredCopyList" If="Exists('Engine/Build/NotForLicensees/EpicInternal.txt')"/>
-			<Tag Files="Engine/Binaries/Win64/NotForLicensees/UE4Editor-DDCUtils.dll" With="#FilteredCopyList" If="Exists('Engine/Build/NotForLicensees/EpicInternal.txt') And '$(HostPlatform)' == 'Win64'"/>
 			<Tag Files="Engine/Binaries/Linux/NotForLicensees/libUE4Editor-DDCUtils.so" With="#FilteredCopyList" If="Exists('Engine/Build/NotForLicensees/EpicInternal.txt')"/>
+			<!-- When running on Windows -->
+			<Tag Files="Engine/Binaries/Win64/NotForLicensees/ShaderCompileWorker.modules" With="#FilteredCopyList" If="Exists('Engine/Build/NotForLicensees/EpicInternal.txt') And '$(HostPlatform)' == 'Win64'"/>
+			<Tag Files="Engine/Binaries/Win64/NotForLicensees/ShaderCompileWorker-DDCUtils.dll" With="#FilteredCopyList" If="Exists('Engine/Build/NotForLicensees/EpicInternal.txt') And '$(HostPlatform)' == 'Win64'"/>
+			<Tag Files="Engine/Binaries/Win64/NotForLicensees/UE4Editor.modules" With="#FilteredCopyList" If="Exists('Engine/Build/NotForLicensees/EpicInternal.txt') And '$(HostPlatform)' == 'Win64'"/>
+			<Tag Files="Engine/Binaries/Win64/NotForLicensees/UE4Editor-DDCUtils.dll" With="#FilteredCopyList" If="Exists('Engine/Build/NotForLicensees/EpicInternal.txt') And '$(HostPlatform)' == 'Win64'"/>
 			<!-- Copy everything to a temporary directory -->
 			<Copy From="$(RootDir)" To="$(DDCDir)" Files="#FilteredCopyList"/>
 			<Copy From="$(CsToolsDir)" To="$(DDCDir)" Files="#Build Tools CS Binaries"/>
