UBT: Fix remote compilation on IOS from installed builds.

* Transfer precompiled binaries for IOS and TVOS to remote Mac.
* Force Mac to be treated as a valid installed platform during remote builds.

diff --git a/Engine/Build/InstalledEngineBuild.xml b/Engine/Build/InstalledEngineBuild.xml
--- a/Engine/Build/InstalledEngineBuild.xml
+++ b/Engine/Build/InstalledEngineBuild.xml
@@ -673,7 +673,7 @@
 		<Property Name="InstalledRequirements" Value="$(InstalledRequirements);#UE4Game Win64;#UE4Game Win64 Unstripped;#UE4Game Win64 Stripped;#UE4Game Win64 Unsigned;#UE4Game Win64 Signed" If="'$(WithWin64)' == true"/>
 		<Property Name="InstalledRequirements" Value="$(InstalledRequirements);#UE4Game Win32;#UE4Game Win32 Unstripped;#UE4Game Win32 Stripped;#UE4Game Win32 Unsigned;#UE4Game Win32 Signed" If="'$(WithWin32)' == true"/>
 		<Property Name="InstalledRequirements" Value="$(InstalledRequirements);#UE4Game Android;#UE4Game Android Unstripped;#UE4Game Android Stripped" If="'$(WithAndroid)' == true"/>
-		<Property Name="InstalledRequirements" Value="$(InstalledRequirements);#UE4Game IOS;#UE4Game IOS Unstripped;#UE4Game IOS Stripped" If="'$(WithIOS)' == true"/>
+		<Property Name="InstalledRequirements" Value="$(InstalledRequirements);#UE4Game IOS;#UE4Game IOS Unstripped;#UE4Game IOS Stripped;#Compile UnrealHeaderTool Mac" If="'$(WithIOS)' == true"/>
 		<Property Name="InstalledRequirements" Value="$(InstalledRequirements);#UE4Game TVOS;#UE4Game TVOS Unstripped;#UE4Game TVOS Stripped" If="'$(WithTVOS)' == true"/>
 		<Property Name="InstalledRequirements" Value="$(InstalledRequirements);#UE4Game Linux;#UE4Game Linux Unstripped;#UE4Game Linux Stripped;#Build Tools Linux" If="'$(WithLinux)' == true"/>
 		<Property Name="InstalledRequirements" Value="$(InstalledRequirements);#UE4Game HTML5" If="'$(WithHTML5)' == true"/>
@@ -721,7 +721,7 @@
 				<Property Name="CopyInstalledExceptions" Value="$(CopyInstalledExceptions);$(CopyAndroidExceptions)"/>
 			</Do>
 			<Do If="'$(WithIOS)' == true">
-				<Tag Files="#UE4Game IOS" Except="#UE4Game IOS Unstripped" With="#Installed Win64"/>
+				<Tag Files="#UE4Game IOS;#Compile UnrealHeaderTool Mac" Except="#UE4Game IOS Unstripped" With="#Installed Win64"/>
 				<Tag Files="#UE4Game IOS Stripped" With="#Saved Output"/>
 				<Property Name="CopyInstalledFilter" Value="$(CopyInstalledFilter);$(CopyIOSFilter)"/>
 				<Property Name="CopyInstalledExceptions" Value="$(CopyInstalledExceptions);$(CopyIOSExceptions)"/>
diff --git a/Engine/Build/Rsync/RsyncEngineInstalled.txt b/Engine/Build/Rsync/RsyncEngineInstalled.txt
new file mode 100644
--- /dev/null
+++ b/Engine/Build/Rsync/RsyncEngineInstalled.txt
@@ -0,0 +1,24 @@
+# Additional rules to apply to the installed engine build
+
+# copy all the precompiled intermediates, and files which are used for detection of installed platforms
++ /Binaries/
++ /Binaries/IOS/
++ /Binaries/IOS/*.target
++ /Binaries/TVOS/
++ /Binaries/TVOS/*.target
++ /Intermediate/
++ /Intermediate/Build/
++ /Intermediate/Build/IOS/
++ /Intermediate/Build/IOS/**
++ /Intermediate/Build/TVOS/
++ /Intermediate/Build/TVOS/**
++ /Plugins/**/Intermediate/
++ /Plugins/**/Intermediate/Build/ 
++ /Plugins/**/Intermediate/Build/IOS/ 
++ /Plugins/**/Intermediate/Build/IOS/**
++ /Plugins/**/Intermediate/Build/TVOS/ 
++ /Plugins/**/Intermediate/Build/TVOS/**
+
+# also copy binaries for UHT and DsymExporter on Mac
++ /Binaries/Mac/
++ /Binaries/Mac/*
diff --git a/Engine/Source/Programs/UnrealBuildTool/System/InstalledPlatformInfo.cs b/Engine/Source/Programs/UnrealBuildTool/System/InstalledPlatformInfo.cs
--- a/Engine/Source/Programs/UnrealBuildTool/System/InstalledPlatformInfo.cs
+++ b/Engine/Source/Programs/UnrealBuildTool/System/InstalledPlatformInfo.cs
@@ -1,4 +1,4 @@
-﻿// Copyright 1998-2018 Epic Games, Inc. All Rights Reserved.
+// Copyright 1998-2018 Epic Games, Inc. All Rights Reserved.
 
 using System;
 using System.Collections.Generic;
@@ -264,6 +264,12 @@ namespace UnrealBuildTool
 		/// <returns>True if supported</returns>
 		public static bool IsValidPlatform(UnrealTargetPlatform Platform, EProjectType ProjectType = EProjectType.Any)
 		{
+			// HACK: For installed builds, we always need to treat Mac as a valid platform for generating project files. When remote building from PC, we won't have all the libraries to do this, so we need to fake it.
+			if(Platform == UnrealTargetPlatform.Mac && ProjectType == EProjectType.Any && BuildHostPlatform.Current.Platform == UnrealTargetPlatform.Mac && UnrealBuildTool.IsEngineInstalled())
+			{
+				return true;
+			}
+
 			return ContainsValidConfiguration(
 				(InstalledPlatformConfiguration CurConfig) =>
 				{
diff --git a/Engine/Source/Programs/UnrealBuildTool/ToolChain/RemoteMac.cs b/Engine/Source/Programs/UnrealBuildTool/ToolChain/RemoteMac.cs
--- a/Engine/Source/Programs/UnrealBuildTool/ToolChain/RemoteMac.cs
+++ b/Engine/Source/Programs/UnrealBuildTool/ToolChain/RemoteMac.cs
@@ -236,7 +236,7 @@ namespace UnrealBuildTool
 			BasicRsyncArguments.Add("--compress");
 			BasicRsyncArguments.Add("--verbose");
 			BasicRsyncArguments.Add(String.Format("--rsh=\"{0} -p {1}\"", RsyncAuthentication, ServerPort));
-			BasicRsyncArguments.Add("--chmod=ug=rwX,o=rxX");
+			BasicRsyncArguments.Add("--chmod=ugo=rwx");
 
 			// Build a list of arguments for Rsync filters
 			CommonRsyncArguments = new List<string>(BasicRsyncArguments);
@@ -401,6 +401,9 @@ namespace UnrealBuildTool
 				// Always generate a .stub
 				RemoteArguments.Add("-CreateStub");
 
+				// Cannot use makefiles, since we need PostBuildSync() to generate the IPA (and that requires a TargetRules instance)
+				RemoteArguments.Add("-NoUBTMakefiles");
+
 				// Get the provisioning data for this project
 				IOSProvisioningData ProvisioningData = ((IOSPlatform)UEBuildPlatform.GetBuildPlatform(TargetDesc.Platform)).ReadProvisioningData(TargetDesc.ProjectFile);
 				if(ProvisioningData == null || ProvisioningData.MobileProvisionFile == null)
@@ -764,6 +767,10 @@ namespace UnrealBuildTool
 			// Upload the engine files
 			List<FileReference> EngineFilters = new List<FileReference>();
 			EngineFilters.Add(ScriptProtectList);
+			if(UnrealBuildTool.IsEngineInstalled())
+			{
+				EngineFilters.Add(FileReference.Combine(UnrealBuildTool.EngineDirectory, "Build", "Rsync", "RsyncEngineInstalled.txt"));
+			}
 			EngineFilters.Add(FileReference.Combine(UnrealBuildTool.EngineDirectory, "Build", "Rsync", "RsyncEngine.txt"));
 
 			Log.TraceInformation("[Remote] Uploading engine files...");
