Merging //UE4/Dev-Main @ 4070185 to Dev-Core (//UE4/Dev-Core)

diff --git a/Engine/Build/InstalledEngineBuild.xml b/Engine/Build/InstalledEngineBuild.xml
--- a/Engine/Build/InstalledEngineBuild.xml
+++ b/Engine/Build/InstalledEngineBuild.xml
@@ -827,7 +827,10 @@
 		<Property Name="FilterRequirements" Value="$(FilterRequirements);#UE4Game IOS;#UE4Game IOS Unstripped;#UE4Game IOS Stripped;#UE4Game IOS Includes" If="'$(WithIOS)' == true"/>
 		<Property Name="FilterRequirements" Value="$(FilterRequirements);#UE4Game TVOS;#UE4Game TVOS Unstripped;#UE4Game TVOS Stripped;#UE4Game TVOS Includes" If="'$(WithTVOS)' == true"/>
 		<Property Name="FilterRequirements" Value="$(FilterRequirements);#UE4Game HTML5;#UE4Game HTML5 Includes" If="'$(WithHTML5)' == true"/>
-		<Property Name="FilterRequirements" Value="$(FilterRequirements);#UE4Game LuminMac;#UE4Game LuminMac Includes;#UE4Game LuminMac Unstripped;#UE4Game LuminMac Stripped;#UE4Game LuminMac Includes" If="'$(WithLumin)' == true"/>
+		<!--	
+		Temporarily removing until we have MLSDK setup on build machines 
+		<Property Name="FilterRequirements" Value="$(FilterRequirements);#UE4Game LuminMac;#UE4Game LuminMac Includes;#UE4Game LuminMac Unstripped;#UE4Game LuminMac Stripped;#UE4Game LuminMac Includes" If="'$(WithLumin)' == true"/> 
+		-->
 		<Property Name="FilterRequirements" Value="$(FilterRequirements);#CompressedDDCMac" If="'$(WithDDC)' == true"/>
 
 		<!-- Feature packs -->
@@ -878,12 +881,15 @@
 				<Property Name="CopyInstalledFilter" Value="$(CopyInstalledFilter);$(CopyHTML5Filter)"/>
 				<Property Name="CopyInstalledExceptions" Value="$(CopyInstalledExceptions);$(CopyHTML5Exceptions)"/>
 			</Do>
+			<!--	
+			Temporarily removing until we have MLSDK setup on build machines 
 			<Do If="'$(WithLumin)' == true">
 				<Tag Files="#UE4Game LuminMac;#UE4Game LuminMac Includes" Except="#UE4Game LuminMac Unstripped" With="#Installed Mac"/>
 				<Tag Files="#UE4Game LuminMac Stripped" With="#Saved Output"/>
 				<Property Name="CopyInstalledFilter" Value="$(CopyInstalledFilter);$(CopyLuminFilter)"/>
 				<Property Name="CopyInstalledExceptions" Value="$(CopyInstalledExceptions);$(CopyLuminExceptions)"/>
 			</Do>
+			-->
 
 			<!-- Tag any dependencies from all previous build products -->
 			<Tag Files="#Installed Mac" Filter="*.target" With="#TargetReceipts"/>
diff --git a/Engine/Documentation/Source/Programming/UnrealBuildSystem/Configuration/BuildConfigProperties/BuildConfigProperties.INT.udn b/Engine/Documentation/Source/Programming/UnrealBuildSystem/Configuration/BuildConfigProperties/BuildConfigProperties.INT.udn
--- a/Engine/Documentation/Source/Programming/UnrealBuildSystem/Configuration/BuildConfigProperties/BuildConfigProperties.INT.udn
+++ b/Engine/Documentation/Source/Programming/UnrealBuildSystem/Configuration/BuildConfigProperties/BuildConfigProperties.INT.udn
@@ -164,6 +164,14 @@ $ bShowXGEMonitor : Whether to display the XGE build monitor.
 
 $ bStopXGECompilationAfterErrors : When enabled, XGE will stop compiling targets after a compile error occurs.  Recommended, as it saves computing resources for others.
 
+### CLionGenerator
+
+$ bIncludeDocumentation : True if we should include documentation in the generated projects
+
+$ MasterProjectName : Name of the master project file (e.g. base file name for the solution file for Visual Studio, or the Xcode project file on Mac)
+
+$ bMasterProjectNameFromFolder : If true, sets the master project name according to the name of the folder it's in
+
 ### CMakefileGenerator
 
 $ bIncludeDocumentation : True if we should include documentation in the generated projects
diff --git a/Engine/Plugins/Lumin/MagicLeap/Source/MLSDK/Public/MagicLeapMath.h b/Engine/Plugins/Lumin/MagicLeap/Source/MLSDK/Public/MagicLeapMath.h
--- a/Engine/Plugins/Lumin/MagicLeap/Source/MLSDK/Public/MagicLeapMath.h
+++ b/Engine/Plugins/Lumin/MagicLeap/Source/MLSDK/Public/MagicLeapMath.h
@@ -37,7 +37,12 @@
 #include "Kismet/KismetMathLibrary.h"
 #include "CoreMinimal.h"
 #if WITH_MLSDK
+THIRD_PARTY_INCLUDES_START // pushes current warning pragmas
+#if PLATFORM_WINDOWS
+	__pragma(warning(disable: 4201)) /* warning C4201: nonstandard extension used: nameless struct/union */
+#endif
 #include <ml_types.h>
+THIRD_PARTY_INCLUDES_END // pops current warning pragmas
 #endif //WITH_MLSDK
 
 namespace MagicLeap
diff --git a/Engine/Plugins/Lumin/MagicLeap/Source/MagicLeap/Private/MagicLeapCustomPresent.cpp b/Engine/Plugins/Lumin/MagicLeap/Source/MagicLeap/Private/MagicLeapCustomPresent.cpp
--- a/Engine/Plugins/Lumin/MagicLeap/Source/MagicLeap/Private/MagicLeapCustomPresent.cpp
+++ b/Engine/Plugins/Lumin/MagicLeap/Source/MagicLeap/Private/MagicLeapCustomPresent.cpp
@@ -318,6 +318,8 @@ void FMagicLeapCustomPresentOpenGL::FinishRendering()
 			glDisable(GL_FRAMEBUFFER_SRGB);
 		}
 
+		check(vp_array.num_virtual_cameras >= 2); // We assume at least one virtual camera per eye
+
 		FVector2D InternalTextureDims;
 		Plugin->GetAppFrameworkConst().GetDeviceResolution(InternalTextureDims);
 
@@ -353,6 +355,8 @@ void FMagicLeapCustomPresentOpenGL::FinishRendering()
 			glEnable(GL_FRAMEBUFFER_SRGB);
 		}
 
+		static_assert(ARRAY_COUNT(vp_array.virtual_cameras) == 2, "The MLSDK has updated the size of the virtual_cameras array.");
+#if 0 // Enable this in case the MLSDK increases the size of the virtual_cameras array past 2
 		for (uint32 i = 2; i < vp_array.num_virtual_cameras; ++i)
 		{
 			bResult = MLGraphicsSignalSyncObjectGL(Plugin->GraphicsClient, vp_array.virtual_cameras[i].sync_object, &OutStatus);
@@ -361,6 +365,7 @@ void FMagicLeapCustomPresentOpenGL::FinishRendering()
 				UE_LOG(LogMagicLeap, Error, TEXT("MLGraphicsSignalSyncObjectGL for eye %d failed with status %d"), i, static_cast<int32>(OutStatus));
 			}
 		}
+#endif
 
 #if BEGIN_END_FRAME_BALANCE_HACK
 		--BalanceCounter;
diff --git a/Engine/Source/Runtime/Core/Private/UObject/UnrealNames.cpp b/Engine/Source/Runtime/Core/Private/UObject/UnrealNames.cpp
--- a/Engine/Source/Runtime/Core/Private/UObject/UnrealNames.cpp
+++ b/Engine/Source/Runtime/Core/Private/UObject/UnrealNames.cpp
@@ -328,7 +328,7 @@ FString FName::NameToDisplayString( const FString& InDisplayName, const bool bIs
 		}
 
 		// If the current character is upper case or a digit, and the previous character wasn't, then we need to insert a space if there wasn't one previously
-		if( (bUpperCase || bIsDigit) && (!bInARun || ((CharIndex + 1 < Chars.Num() && FChar::IsLower(Chars[CharIndex + 1])) && !bIsDigit)) && !bWasOpenParen)
+		if( (bUpperCase || bIsDigit) && !bInARun && !bWasOpenParen)
 		{
 			if( !bWasSpace && OutDisplayName.Len() > 0 )
 			{
diff --git a/Engine/Source/Runtime/WebBrowser/WebBrowser.Build.cs b/Engine/Source/Runtime/WebBrowser/WebBrowser.Build.cs
--- a/Engine/Source/Runtime/WebBrowser/WebBrowser.Build.cs
+++ b/Engine/Source/Runtime/WebBrowser/WebBrowser.Build.cs
@@ -24,10 +24,11 @@ public class WebBrowser : ModuleRules
         );
 
         if (Target.Platform == UnrealTargetPlatform.Android ||
-            Target.Platform == UnrealTargetPlatform.IOS)
+            Target.Platform == UnrealTargetPlatform.IOS ||
+            Target.Platform == UnrealTargetPlatform.TVOS)
 		{
-			// We need these on Android for external texture support
-			PrivateDependencyModuleNames.AddRange(
+            // We need these on mobile for external texture support
+            PrivateDependencyModuleNames.AddRange(
 				new string[]
 				{
 					"WebBrowserTexture",
