UBT: Simplify logic for precompiling binaries.

* Target no longer has separate list of "precompile only" binaries or modules. New -AllModules option allows adding every module to a target, which can be used with -Precompile and -NoLink to precompile object files for monolithic builds.
* Precompiled file lists have been removed from target receipts.
* The manifest now includes all generated headers and precompiled files when run with the -Precompile option.
* Separate -DependencyList=Foo.txt has been added to write a list of all dependencies required to use precompiled binaries. This file list can be read using the <Tag> task in buildgraph.

#rnx

diff --git a/Engine/Build/InstalledEngineBuild.xml b/Engine/Build/InstalledEngineBuild.xml
--- a/Engine/Build/InstalledEngineBuild.xml
+++ b/Engine/Build/InstalledEngineBuild.xml
@@ -118,9 +118,9 @@
 		<Node Name="Compile UnrealHeaderTool Win64" Requires="Update Version Files">
 			<Compile Target="UnrealHeaderTool" Platform="Win64" Configuration="Development" Arguments="-precompile"/>
 		</Node>
-		<Node Name="Compile UE4Editor Win64" Requires="Compile UnrealHeaderTool Win64" Produces="#UE4Editor Win64;#UE4Editor Win64 Unstripped;#UE4Editor Win64 Stripped;#UE4Editor Win64 Unsigned;#UE4Editor Win64 Signed;#UE4Editor Win64 Includes">
-			<Compile Target="UE4Editor" Platform="Win64" Configuration="Development" Tag="#UE4Editor Win64" Arguments="-precompile"/>
-			<Compile Target="UE4Editor" Platform="Win64" Configuration="DebugGame" Tag="#UE4Editor Win64" Arguments="-precompile"/>
+		<Node Name="Compile UE4Editor Win64" Requires="Compile UnrealHeaderTool Win64" Produces="#UE4Editor Win64;#UE4Editor Win64 Unstripped;#UE4Editor Win64 Stripped;#UE4Editor Win64 Unsigned;#UE4Editor Win64 Signed">
+			<Compile Target="UE4Editor" Platform="Win64" Configuration="Development" Tag="#UE4Editor Win64" Arguments="-precompile -allmodules"/>
+			<Compile Target="UE4Editor" Platform="Win64" Configuration="DebugGame" Tag="#UE4Editor Win64" Arguments="-precompile -allmodules"/>
 
 			<Do If="$(EmbedSrcSrvInfo)">
 				<!-- Embed source info into the PDB files. Should be done from this machine to ensure that paths are correct. -->
@@ -148,8 +148,6 @@
 				<Sign Files="#Editor_ToSign_Win64" Tag="#UE4Editor Win64 Signed"/>
 			</Do>
 
-			<!-- Tag the generated includes for this target -->
-			<Tag Files="Engine/Intermediate/Build/Win64/UE4Editor/Inc/...;Engine/Plugins/.../Intermediate/Build/Win64/UE4Editor/Inc/..." Except=".../Timestamp" With="#UE4Editor Win64 Includes"/>
 		</Node>
 	</Agent>
 
@@ -157,11 +155,11 @@
 	<!-- Base editor and tools for compile on Mac -->
 	<Agent Name="Editor Mac" Type="CompileMac;Mac">
 		<Node Name="Compile UnrealHeaderTool Mac" Requires="Update Version Files">
-			<Compile Target="UnrealHeaderTool" Platform="Mac" Configuration="Development" Arguments="-precompile" If="'$(HostPlatform)' != 'Win64'"/>
+			<Compile Target="UnrealHeaderTool" Platform="Mac" Configuration="Development" Arguments="-precompile -allmodules" If="'$(HostPlatform)' != 'Win64'"/>
 		</Node>
-		<Node Name="Compile UE4Editor Mac" Requires="Compile UnrealHeaderTool Mac" Produces="#UE4Editor Mac;#UE4Editor Mac Unstripped;#UE4Editor Mac Stripped;#UE4Editor Mac Unsigned;#UE4Editor Mac Signed;#UE4Editor Mac Includes">
-			<Compile Target="UE4Editor" Platform="Mac" Configuration="Development" Tag="#UE4Editor Mac" Arguments="-precompile"/>
-			<Compile Target="UE4Editor" Platform="Mac" Configuration="DebugGame" Tag="#UE4Editor Mac" Arguments="-precompile"/>
+		<Node Name="Compile UE4Editor Mac" Requires="Compile UnrealHeaderTool Mac" Produces="#UE4Editor Mac;#UE4Editor Mac Unstripped;#UE4Editor Mac Stripped;#UE4Editor Mac Unsigned;#UE4Editor Mac Signed">
+			<Compile Target="UE4Editor" Platform="Mac" Configuration="Development" Tag="#UE4Editor Mac" Arguments="-precompile -allmodules"/>
+			<Compile Target="UE4Editor" Platform="Mac" Configuration="DebugGame" Tag="#UE4Editor Mac" Arguments="-precompile -allmodules"/>
 
 			<Tag Files="#UE4Editor Mac" Filter="$(MacStripFilter)" With="#UE4Editor Mac Unstripped"/>
 
@@ -181,20 +179,17 @@
 				<Copy Files="#UE4Editor Mac Unsigned" From="$(RootDir)" To="$(SavedOutput)" Tag="#Editor_ToSign_Mac"/>
 				<Sign Files="#Editor_ToSign_Mac" Tag="#UE4Editor Mac Signed"/>
 			</Do>
-
-			<!-- Tag the generated includes for this target -->
-			<Tag Files="Engine/Intermediate/Build/Mac/UE4Editor/Inc/...;Engine/Plugins/.../Intermediate/Build/Mac/UE4Editor/Inc/..." Except=".../Timestamp" With="#UE4Editor Mac Includes"/>
 		</Node>
 	</Agent>
 
 	<!-- Base editor and tools for compile on Linux -->
 	<Agent Name="Editor Linux" Type="CompileWin64;Win64">
 		<Node Name="Compile UnrealHeaderTool Linux" Requires="Update Version Files">
-			<Compile Target="UnrealHeaderTool" Platform="Linux" Configuration="Development" Arguments="-precompile"/>
+			<Compile Target="UnrealHeaderTool" Platform="Linux" Configuration="Development" Arguments="-precompile -allmodules"/>
 		</Node>
-		<Node Name="Compile UE4Editor Linux" Requires="Compile UnrealHeaderTool $(HostPlatform)" Produces="#UE4Editor Linux;#UE4Editor Linux Unstripped;#UE4Editor Linux Stripped;#UE4Editor Linux Includes">
-			<Compile Target="UE4Editor" Platform="Linux" Configuration="Development" Tag="#UE4Editor Linux" Arguments="-precompile"/>
-			<Compile Target="UE4Editor" Platform="Linux" Configuration="DebugGame" Tag="#UE4Editor Linux" Arguments="-precompile"/>
+		<Node Name="Compile UE4Editor Linux" Requires="Compile UnrealHeaderTool $(HostPlatform)" Produces="#UE4Editor Linux;#UE4Editor Linux Unstripped;#UE4Editor Linux Stripped">
+			<Compile Target="UE4Editor" Platform="Linux" Configuration="Development" Tag="#UE4Editor Linux" Arguments="-precompile -allmodules"/>
+			<Compile Target="UE4Editor" Platform="Linux" Configuration="DebugGame" Tag="#UE4Editor Linux" Arguments="-precompile -allmodules"/>
 
 			<Tag Files="#UE4Editor Linux" Filter="$(LinuxStripFilter)" With="#UE4Editor Linux Unstripped"/>
 
@@ -208,18 +203,16 @@
 					<Strip BaseDir="$(RootDir)" OutputDir="$(SavedOutput)" Platform="Linux" Files="#UE4Editor Linux Unstripped" Tag="#UE4Editor Linux Stripped"/>
 				</Default>
 			</Switch>
-
-			<!-- Tag the generated includes for this target -->
-			<Tag Files="Engine/Intermediate/Build/Linux/B4D820EA/UE4Editor/Inc/...;Engine/Plugins/.../Intermediate/Build/Linux/B4D820EA/UE4Editor/Inc/..." Except=".../Timestamp" With="#UE4Editor Linux Includes"/>
 		</Node>
 	</Agent>
 
 	<!-- Target Platforms normally built from a Win64 machine -->
 	<Agent Name="Target Platforms Win64" Type="CompileWin64;Win64">
 		<!-- Downstream monolithics and tools Win64 -->
-		<Node Name="Compile UE4Game Win64" Requires="Compile UnrealHeaderTool Win64" Produces="#UE4Game Win64;#UE4Game Win64 Unstripped;#UE4Game Win64 Stripped;#UE4Game Win64 Unsigned;#UE4Game Win64 Signed;#UE4Game Win64 Includes">
+		<Node Name="Compile UE4Game Win64" Requires="Compile UnrealHeaderTool Win64" Produces="#UE4Game Win64;#UE4Game Win64 Unstripped;#UE4Game Win64 Stripped;#UE4Game Win64 Unsigned;#UE4Game Win64 Signed">
 		
 			<ForEach Name="Configuration" Values="$(GameConfigurations)">
+				<Compile Target="UE4Game" Platform="Win64" Configuration="$(Configuration)" Tag="#UE4Game Win64" Arguments="-precompile -allmodules -nolink $(VSCompilerArg) $(TargetDebugInfoArg)"/>
 				<Compile Target="UE4Game" Platform="Win64" Configuration="$(Configuration)" Tag="#UE4Game Win64" Arguments="-precompile $(VSCompilerArg) $(TargetDebugInfoArg)"/>
 			</ForEach>
 
@@ -248,15 +241,13 @@
 				<Copy Files="#UE4Game Win64 Unsigned" From="$(RootDir)" To="$(SavedOutput)" Tag="#Game_ToSign_Win64"/>
 				<Sign Files="#Game_ToSign_Win64" Tag="#UE4Game Win64 Signed"/>
 			</Do>
-
-			<!-- Tag the generated includes for this target -->
-			<Tag Files="Engine/Intermediate/Build/Win64/UE4/Inc/...;Engine/Plugins/.../Intermediate/Build/Win64/UE4/Inc/..." With="#UE4Game Win64 Includes"/>
 		</Node>
 
 		<!-- Downstream monolithics and tools Win32 -->
-		<Node Name="Compile UE4Game Win32" Requires="Compile UnrealHeaderTool Win64" Produces="#UE4Game Win32;#UE4Game Win32 Unstripped;#UE4Game Win32 Stripped;#UE4Game Win32 Unsigned;#UE4Game Win32 Signed;#UE4Game Win32 Includes">
+		<Node Name="Compile UE4Game Win32" Requires="Compile UnrealHeaderTool Win64" Produces="#UE4Game Win32;#UE4Game Win32 Unstripped;#UE4Game Win32 Stripped;#UE4Game Win32 Unsigned;#UE4Game Win32 Signed">
 		
 			<ForEach Name="Configuration" Values="$(GameConfigurations)">
+				<Compile Target="UE4Game" Platform="Win32" Configuration="$(Configuration)" Tag="#UE4Game Win32" Arguments="-precompile -allmodules -nolink $(VSCompilerArg) $(TargetDebugInfoArg)"/>
 				<Compile Target="UE4Game" Platform="Win32" Configuration="$(Configuration)" Tag="#UE4Game Win32" Arguments="-precompile $(VSCompilerArg) $(TargetDebugInfoArg)"/>
 			</ForEach>
 
@@ -286,103 +277,94 @@
 				<Copy Files="#UE4Game Win32 Unsigned" From="$(RootDir)" To="$(SavedOutput)" Tag="#Game_ToSign_Win32"/>
 				<Sign Files="#Game_ToSign_Win32" Tag="#UE4Game Win32 Signed"/>
 			</Do>
-
-			<!-- Tag the generated includes for this target -->
-			<Tag Files="Engine/Intermediate/Build/Win32/UE4/Inc/...;Engine/Plugins/.../Intermediate/Build/Win32/UE4/Inc/..." With="#UE4Game Win32 Includes"/>
 		</Node>
 
 		<!-- Downstream monolithics and tools Linux -->
-		<Node Name="Compile UE4Game Linux" Requires="Compile UnrealHeaderTool $(HostPlatform)" Produces="#UE4Game Linux;#UE4Game Linux Unstripped;#UE4Game Linux Stripped;#UE4Game Linux Includes">
+		<Node Name="Compile UE4Game Linux" Requires="Compile UnrealHeaderTool $(HostPlatform)" Produces="#UE4Game Linux;#UE4Game Linux Unstripped;#UE4Game Linux Stripped">
 		
 			<ForEach Name="Configuration" Values="$(GameConfigurations)">
+				<Compile Target="UE4Game" Platform="Linux" Configuration="$(Configuration)" Tag="#UE4Game Linux" Arguments="-precompile -allmodules -nolink $(TargetDebugInfoArg)"/>
 				<Compile Target="UE4Game" Platform="Linux" Configuration="$(Configuration)" Tag="#UE4Game Linux" Arguments="-precompile $(TargetDebugInfoArg)"/>
 			</ForEach>
 			
 			<Tag Files="#UE4Game Linux" Filter="$(LinuxStripFilter)" With="#UE4Game Linux Unstripped"/>
 			<Strip BaseDir="$(RootDir)" OutputDir="$(SavedOutput)" Platform="Linux" Files="#UE4Game Linux Unstripped" Tag="#UE4Game Linux Stripped"/>
 
-			<!-- Tag the generated includes for this target -->
-			<Tag Files="Engine/Intermediate/Build/Linux/B4D820EA/UE4/Inc/...;Engine/Plugins/.../Intermediate/Build/Linux/B4D820EA/UE4/Inc/..." With="#UE4Game Linux Includes"/>
 		</Node>
 
 		<!-- Downstream monolithics and tools Android-->
-		<Node Name="Compile UE4Game Android" Requires="Compile UnrealHeaderTool $(HostPlatform)" Produces="#UE4Game Android;#UE4Game Android Unstripped;#UE4Game Android Stripped;#UE4Game Android Includes">
+		<Node Name="Compile UE4Game Android" Requires="Compile UnrealHeaderTool $(HostPlatform)" Produces="#UE4Game Android;#UE4Game Android Unstripped;#UE4Game Android Stripped">
 		
 			<ForEach Name="Configuration" Values="$(GameConfigurations)">
+				<Compile Target="UE4Game" Platform="Android" Configuration="$(Configuration)" Tag="#UE4Game Android" Arguments="-precompile -allmodules -nolink -architectures=$(AndroidArchitectures) -gpuarchitectures=$(AndroidGPUArchitectures) $(TargetDebugInfoArg)"/>
 				<Compile Target="UE4Game" Platform="Android" Configuration="$(Configuration)" Tag="#UE4Game Android" Arguments="-precompile -architectures=$(AndroidArchitectures) -gpuarchitectures=$(AndroidGPUArchitectures) $(TargetDebugInfoArg)"/>
 			</ForEach>
 		
 			<Tag Files="#UE4Game Android" Filter="$(AndroidStripFilter)" With="#UE4Game Android Unstripped"/>
 			<Strip BaseDir="$(RootDir)" OutputDir="$(SavedOutput)" Platform="Android" Files="#UE4Game Android Unstripped" Tag="#UE4Game Android Stripped"/>
 
-			<!-- Tag the generated includes for this target -->
-			<Tag Files="Engine/Intermediate/Build/Android/UE4/Inc/...;Engine/Plugins/.../Intermediate/Build/Android/UE4/Inc/..." With="#UE4Game Android Includes"/>
 		</Node>
 
 		<!-- Downstream monolithics and tools HTML5-->
-		<Node Name="Compile UE4Game HTML5" Requires="Compile UnrealHeaderTool $(HostPlatform)" Produces="#UE4Game HTML5;#UE4Game HTML5 Includes">
+		<Node Name="Compile UE4Game HTML5" Requires="Compile UnrealHeaderTool $(HostPlatform)" Produces="#UE4Game HTML5">
 		
 			<ForEach Name="Configuration" Values="$(GameConfigurations)">
+				<Compile Target="UE4Game" Platform="HTML5" Configuration="$(Configuration)" Tag="#UE4Game HTML5" Arguments="-precompile -allmodules -nolink $(TargetDebugInfoArg)"/>
 				<Compile Target="UE4Game" Platform="HTML5" Configuration="$(Configuration)" Tag="#UE4Game HTML5" Arguments="-precompile $(TargetDebugInfoArg)"/>
 			</ForEach>
 		
-			<!-- Tag the generated includes for this target -->
-			<Tag Files="Engine/Intermediate/Build/HTML5/UE4/Inc/...;Engine/Plugins/.../Intermediate/Build/HTML5/UE4/Inc/..." With="#UE4Game HTML5 Includes"/>
 		</Node>
 
 		<!-- Downstream monolithics and tools PS4 -->
-		<Node Name="Compile UE4Game PS4" Requires="Compile UnrealHeaderTool Win64" Produces="#UE4Game PS4;#UE4Game PS4 Includes">
+		<Node Name="Compile UE4Game PS4" Requires="Compile UnrealHeaderTool Win64" Produces="#UE4Game PS4">
 		
 			<ForEach Name="Configuration" Values="$(GameConfigurations)">
+				<Compile Target="UE4Game" Platform="PS4" Configuration="$(Configuration)" Tag="#UE4Game PS4" Arguments="-precompile -allmodules -nolink $(TargetDebugInfoArg)"/>
 				<Compile Target="UE4Game" Platform="PS4" Configuration="$(Configuration)" Tag="#UE4Game PS4" Arguments="-precompile $(TargetDebugInfoArg)"/>
 			</ForEach>
 
-			<!-- Tag the generated includes for this target -->
-			<Tag Files="Engine/Intermediate/Build/PS4/UE4/Inc/...;Engine/Plugins/.../Intermediate/Build/PS4/UE4/Inc/..." With="#UE4Game PS4 Includes"/>
 		</Node>
 
 		<!-- Downstream monolithics and tools XboxOne -->
-		<Node Name="Compile UE4Game XboxOne" Requires="Compile UnrealHeaderTool Win64" Produces="#UE4Game XboxOne;#UE4Game XboxOne Includes">
+		<Node Name="Compile UE4Game XboxOne" Requires="Compile UnrealHeaderTool Win64" Produces="#UE4Game XboxOne">
 		
 			<ForEach Name="Configuration" Values="$(GameConfigurations)">
+				<Compile Target="UE4Game" Platform="XboxOne" Configuration="$(Configuration)" Tag="#UE4Game XboxOne" Arguments="-precompile -allmodules -nolink $(TargetDebugInfoArg)"/>
 				<Compile Target="UE4Game" Platform="XboxOne" Configuration="$(Configuration)" Tag="#UE4Game XboxOne" Arguments="-precompile $(TargetDebugInfoArg)"/>
 			</ForEach>
 			
-			<!-- Tag the generated includes for this target -->
-			<Tag Files="Engine/Intermediate/Build/XboxOne/UE4/Inc/...;Engine/Plugins/.../Intermediate/Build/XboxOne/UE4/Inc/..." With="#UE4Game XboxOne Includes"/>
 		</Node>
 
 		<!-- Downstream monolithics Switch -->
-		<Node Name="Compile UE4Game Switch" Requires="Compile UnrealHeaderTool Win64" Produces="#UE4Game Switch;#UE4Game Switch Includes" If="'$(WithSwitch)' == true">
+		<Node Name="Compile UE4Game Switch" Requires="Compile UnrealHeaderTool Win64" Produces="#UE4Game Switch" If="'$(WithSwitch)' == true">
 		
 			<ForEach Name="Configuration" Values="$(GameConfigurations)">
+				<Compile Target="UE4Game" Platform="Switch" Configuration="$(Configuration)" Tag="#UE4Game Switch" Arguments="-precompile -allmodules -nolink $(TargetDebugInfoArg)"/>
 				<Compile Target="UE4Game" Platform="Switch" Configuration="$(Configuration)" Tag="#UE4Game Switch" Arguments="-precompile $(TargetDebugInfoArg)"/>
 			</ForEach>
 			
-			<!-- Tag the generated includes for this target -->
-			<Tag Files="Engine/Intermediate/Build/Switch/UE4/Inc/...;Engine/Plugins/.../Intermediate/Build/Switch/UE4/Inc/..." With="#UE4Game Switch Includes"/>
 		</Node>
 		
-		<Node Name="Compile UE4Game Lumin" Requires="Compile UnrealHeaderTool Win64" Produces="#UE4Game Lumin;#UE4Game Lumin Unstripped;#UE4Game Lumin Stripped;#UE4Game Lumin Includes" If="'$(WithLumin)' == true">
+		<Node Name="Compile UE4Game Lumin" Requires="Compile UnrealHeaderTool Win64" Produces="#UE4Game Lumin;#UE4Game Lumin Unstripped;#UE4Game Lumin Stripped" If="'$(WithLumin)' == true">
 
 			<ForEach Name="Configuration" Values="$(GameConfigurations)">
+				<Compile Target="UE4Game" Platform="Lumin" Configuration="$(Configuration)" Tag="#UE4Game Lumin" Arguments="-precompile -allmodules -nolink -architectures=$(AndroidArchitectures) -gpuarchitectures=$(LuminGPUArchitectures) $(TargetDebugInfoArg)" AllowParallelExecutor="false"/>
 				<Compile Target="UE4Game" Platform="Lumin" Configuration="$(Configuration)" Tag="#UE4Game Lumin" Arguments="-precompile -architectures=$(AndroidArchitectures) -gpuarchitectures=$(LuminGPUArchitectures) $(TargetDebugInfoArg)" AllowParallelExecutor="false"/>
 			</ForEach>
 			
 			<Tag Files="#UE4Game Lumin" Filter="$(LuminStripFilter)" With="#UE4Game Lumin Unstripped"/>
 			<Strip BaseDir="$(RootDir)" OutputDir="$(SavedOutput)" Platform="Lumin" Files="#UE4Game Lumin Unstripped" Tag="#UE4Game Lumin Stripped"/>
 
-			<!-- Tag the generated includes for this target -->
-			<Tag Files="Engine/Intermediate/Build/Lumin/UE4/Inc/...;Engine/Plugins/.../Intermediate/Build/Lumin/UE4/Inc/..." With="#UE4Game Lumin Includes"/>
 		</Node>
 	</Agent>
 
 	<!-- Target Platforms normally built from a Mac machine -->
 	<Agent Name="Target Platforms Mac" Type="CompileMac;Mac">
 		<!-- Downstream monolithics and tools Mac -->
-		<Node Name="Compile UE4Game Mac" Requires="Compile UnrealHeaderTool Mac" Produces="#UE4Game Mac;#UE4Game Mac Unstripped;#UE4Game Mac Stripped;#UE4Game Mac Unsigned;#UE4Game Mac Signed;#UE4Game Mac Includes">
+		<Node Name="Compile UE4Game Mac" Requires="Compile UnrealHeaderTool Mac" Produces="#UE4Game Mac;#UE4Game Mac Unstripped;#UE4Game Mac Stripped;#UE4Game Mac Unsigned;#UE4Game Mac Signed">
 		
 			<ForEach Name="Configuration" Values="$(GameConfigurations)">
+				<Compile Target="UE4Game" Platform="Mac" Configuration="$(Configuration)" Tag="#UE4Game Mac" Arguments="-precompile -allmodules -nolink $(TargetDebugInfoArg)"/>
 				<Compile Target="UE4Game" Platform="Mac" Configuration="$(Configuration)" Tag="#UE4Game Mac" Arguments="-precompile $(TargetDebugInfoArg)"/>
 			</ForEach>
 			
@@ -395,47 +377,45 @@
 				<Sign Files="#Game_ToSign_Mac" Tag="#UE4Game Mac Signed"/>
 			</Do>
 
-			<!-- Tag the generated includes for this target -->
-			<Tag Files="Engine/Intermediate/Build/Mac/UE4/Inc/...;Engine/Plugins/.../Intermediate/Build/Mac/UE4/Inc/..." With="#UE4Game Mac Includes"/>
 		</Node>
 
 		<!-- Downstream monolithics and tools IOS-->
-		<Node Name="Compile UE4Game IOS" Requires="Compile UnrealHeaderTool Mac" Produces="#UE4Game IOS;#UE4Game IOS Unstripped;#UE4Game IOS Stripped;#UE4Game IOS Includes">
+		<Node Name="Compile UE4Game IOS" Requires="Compile UnrealHeaderTool Mac" Produces="#UE4Game IOS;#UE4Game IOS Unstripped;#UE4Game IOS Stripped">
 		
 			<ForEach Name="Configuration" Values="$(GameConfigurations)">
+				<Compile Target="UE4Game" Platform="IOS" Configuration="$(Configuration)" Tag="#UE4Game IOS" Arguments="-precompile -allmodules -nolink $(TargetDebugInfoArg)"/>
 				<Compile Target="UE4Game" Platform="IOS" Configuration="$(Configuration)" Tag="#UE4Game IOS" Arguments="-precompile $(TargetDebugInfoArg)"/>
 			</ForEach>
 
 			<Tag Files="#UE4Game IOS" Filter="$(IOSStripFilter)" With="#UE4Game IOS Unstripped"/>
 			<Strip BaseDir="$(RootDir)" OutputDir="$(SavedOutput)" Platform="IOS" Files="#UE4Game IOS Unstripped" Tag="#UE4Game IOS Stripped"/>
 
-			<!-- Tag the generated includes for this target -->
-			<Tag Files="Engine/Intermediate/Build/IOS/UE4/Inc/...;Engine/Plugins/.../Intermediate/Build/IOS/UE4/Inc/..." With="#UE4Game IOS Includes"/>
 		</Node>
 
 		<!-- Downstream monolithics and tools TVOS-->
-		<Node Name="Compile UE4Game TVOS" Requires="Compile UnrealHeaderTool Mac" Produces="#UE4Game TVOS;#UE4Game TVOS Unstripped;#UE4Game TVOS Stripped;#UE4Game TVOS Includes">
+		<Node Name="Compile UE4Game TVOS" Requires="Compile UnrealHeaderTool Mac" Produces="#UE4Game TVOS;#UE4Game TVOS Unstripped;#UE4Game TVOS Stripped">
 		
 			<ForEach Name="Configuration" Values="$(GameConfigurations)">
+				<Compile Target="UE4Game" Platform="TVOS" Configuration="$(Configuration)" Tag="#UE4Game TVOS" Arguments="-precompile -allmodules -nolink $(TargetDebugInfoArg)"/>
 				<Compile Target="UE4Game" Platform="TVOS" Configuration="$(Configuration)" Tag="#UE4Game TVOS" Arguments="-precompile $(TargetDebugInfoArg)"/>
 			</ForEach>
 
 			<Tag Files="#UE4Game TVOS" Filter="$(TVOSStripFilter)" With="#UE4Game TVOS Unstripped"/>
 			<Strip BaseDir="$(RootDir)" OutputDir="$(SavedOutput)" Platform="TVOS" Files="#UE4Game TVOS Unstripped" Tag="#UE4Game TVOS Stripped"/>
 
-			<!-- Tag the generated includes for this target -->
-			<Tag Files="Engine/Intermediate/Build/TVOS/UE4/Inc/...;Engine/Plugins/.../Intermediate/Build/TVOS/UE4/Inc/..." With="#UE4Game TVOS Includes"/>
 		</Node>
 		
 		<!-- Downstream monolithics and tools Lumin-->
-		<Node Name="Compile UE4Game LuminMac" Requires="Compile UnrealHeaderTool Mac" Produces="#UE4Game LuminMac;#UE4Game LuminMac Unstripped;#UE4Game LuminMac Stripped;#UE4Game LuminMac Includes">
-			<Compile Target="UE4Game" Platform="Lumin" Configuration="Development" Tag="#UE4Game LuminMac" Arguments="-precompile -architectures=$(AndroidArchitectures) -gpuarchitectures=$(LuminGPUArchitectures)" AllowParallelExecutor="false"/>
-			<Compile Target="UE4Game" Platform="Lumin" Configuration="Shipping" Tag="#UE4Game LuminMac" Arguments="-precompile -architectures=$(AndroidArchitectures) -gpuarchitectures=$(LuminGPUArchitectures)" AllowParallelExecutor="false"/>
+		<Node Name="Compile UE4Game LuminMac" Requires="Compile UnrealHeaderTool Mac" Produces="#UE4Game LuminMac;#UE4Game LuminMac Unstripped;#UE4Game LuminMac Stripped">
+
+			<ForEach Name="Configuration" Values="$(GameConfigurations)">
+				<Compile Target="UE4Game" Platform="Lumin" Configuration="$(Configuration)" Tag="#UE4Game LuminMac" Arguments="-precompile -allmodules -nolink -architectures=$(AndroidArchitectures) -gpuarchitectures=$(LuminGPUArchitectures)" AllowParallelExecutor="false"/>
+				<Compile Target="UE4Game" Platform="Lumin" Configuration="$(Configuration)" Tag="#UE4Game LuminMac" Arguments="-precompile -architectures=$(AndroidArchitectures) -gpuarchitectures=$(LuminGPUArchitectures)" AllowParallelExecutor="false"/>
+			</ForEach>
 
 			<Tag Files="#UE4Game LuminMac" Filter="$(LuminStripFilter)" With="#UE4Game LuminMac Unstripped"/>
 			<Strip BaseDir="$(RootDir)" OutputDir="$(SavedOutput)" Platform="Lumin" Files="#UE4Game LuminMac Unstripped" Tag="#UE4Game LuminMac Stripped"/>
 
-			<Tag Files="Engine/Intermediate/Build/Lumin/UE4/Inc/...;Engine/Plugins/.../Intermediate/Build/Lumin/UE4/Inc/..." With="#UE4Game LuminMac Includes"/>
 		</Node>
 	</Agent>
 
@@ -668,19 +648,19 @@
 
 	<Agent Name="Installed Build Group Win64" Type="Win64">
 		<!-- Build up a list of requirements as some can be optional -->
-		<Property Name="InstalledRequirements" Value="Compile UnrealHeaderTool Win64;#UE4Editor Win64;#UE4Editor Win64 Unstripped;#UE4Editor Win64 Stripped;#UE4Editor Win64 Unsigned;#UE4Editor Win64 Signed;#UE4Editor Win64 Includes;#Build Tools Win64;#Build Tools CS Binaries"/>
+		<Property Name="InstalledRequirements" Value="Compile UnrealHeaderTool Win64;#UE4Editor Win64;#UE4Editor Win64 Unstripped;#UE4Editor Win64 Stripped;#UE4Editor Win64 Unsigned;#UE4Editor Win64 Signed;#Build Tools Win64;#Build Tools CS Binaries"/>
 		<!-- Optional Target Platform requirements -->
-		<Property Name="InstalledRequirements" Value="$(InstalledRequirements);#UE4Game Win64;#UE4Game Win64 Unstripped;#UE4Game Win64 Stripped;#UE4Game Win64 Unsigned;#UE4Game Win64 Signed;#UE4Game Win64 Includes" If="'$(WithWin64)' == true"/>
-		<Property Name="InstalledRequirements" Value="$(InstalledRequirements);#UE4Game Win32;#UE4Game Win32 Unstripped;#UE4Game Win32 Stripped;#UE4Game Win32 Unsigned;#UE4Game Win32 Signed;#UE4Game Win32 Includes" If="'$(WithWin32)' == true"/>
-		<Property Name="InstalledRequirements" Value="$(InstalledRequirements);#UE4Game Android;#UE4Game Android Unstripped;#UE4Game Android Stripped;#UE4Game Android Includes" If="'$(WithAndroid)' == true"/>
-		<Property Name="InstalledRequirements" Value="$(InstalledRequirements);#UE4Game IOS;#UE4Game IOS Unstripped;#UE4Game IOS Stripped;#UE4Game IOS Includes" If="'$(WithIOS)' == true"/>
-		<Property Name="InstalledRequirements" Value="$(InstalledRequirements);#UE4Game TVOS;#UE4Game TVOS Unstripped;#UE4Game TVOS Stripped;#UE4Game TVOS Includes" If="'$(WithTVOS)' == true"/>
-		<Property Name="InstalledRequirements" Value="$(InstalledRequirements);#UE4Game Linux;#UE4Game Linux Unstripped;#UE4Game Linux Stripped;#UE4Game Linux Includes;#Build Tools Linux" If="'$(WithLinux)' == true"/>
-		<Property Name="InstalledRequirements" Value="$(InstalledRequirements);#UE4Game HTML5;#UE4Game HTML5 Includes" If="'$(WithHTML5)' == true"/>
-		<Property Name="InstalledRequirements" Value="$(InstalledRequirements);#UE4Game PS4;#UE4Game PS4 Includes" If="'$(WithPS4)' == true"/>
-		<Property Name="InstalledRequirements" Value="$(InstalledRequirements);#UE4Game XboxOne;#UE4Game XboxOne Includes" If="'$(WithXboxOne)' == true"/>
-		<Property Name="InstalledRequirements" Value="$(InstalledRequirements);#UE4Game Switch;#UE4Game Switch Includes" If="'$(WithSwitch)' == true"/>
-		<Property Name="InstalledRequirements" Value="$(InstalledRequirements);#UE4Game Lumin;#UE4Game Lumin Includes" If="'$(WithLumin)' == true"/>
+		<Property Name="InstalledRequirements" Value="$(InstalledRequirements);#UE4Game Win64;#UE4Game Win64 Unstripped;#UE4Game Win64 Stripped;#UE4Game Win64 Unsigned;#UE4Game Win64 Signed" If="'$(WithWin64)' == true"/>
+		<Property Name="InstalledRequirements" Value="$(InstalledRequirements);#UE4Game Win32;#UE4Game Win32 Unstripped;#UE4Game Win32 Stripped;#UE4Game Win32 Unsigned;#UE4Game Win32 Signed" If="'$(WithWin32)' == true"/>
+		<Property Name="InstalledRequirements" Value="$(InstalledRequirements);#UE4Game Android;#UE4Game Android Unstripped;#UE4Game Android Stripped" If="'$(WithAndroid)' == true"/>
+		<Property Name="InstalledRequirements" Value="$(InstalledRequirements);#UE4Game IOS;#UE4Game IOS Unstripped;#UE4Game IOS Stripped" If="'$(WithIOS)' == true"/>
+		<Property Name="InstalledRequirements" Value="$(InstalledRequirements);#UE4Game TVOS;#UE4Game TVOS Unstripped;#UE4Game TVOS Stripped" If="'$(WithTVOS)' == true"/>
+		<Property Name="InstalledRequirements" Value="$(InstalledRequirements);#UE4Game Linux;#UE4Game Linux Unstripped;#UE4Game Linux Stripped;#Build Tools Linux" If="'$(WithLinux)' == true"/>
+		<Property Name="InstalledRequirements" Value="$(InstalledRequirements);#UE4Game HTML5" If="'$(WithHTML5)' == true"/>
+		<Property Name="InstalledRequirements" Value="$(InstalledRequirements);#UE4Game PS4" If="'$(WithPS4)' == true"/>
+		<Property Name="InstalledRequirements" Value="$(InstalledRequirements);#UE4Game XboxOne" If="'$(WithXboxOne)' == true"/>
+		<Property Name="InstalledRequirements" Value="$(InstalledRequirements);#UE4Game Switch" If="'$(WithSwitch)' == true"/>
+		<Property Name="InstalledRequirements" Value="$(InstalledRequirements);#UE4Game Lumin" If="'$(WithLumin)' == true"/>
 		<Property Name="InstalledRequirements" Value="$(InstalledRequirements);#CompressedDDCWin64" If="'$(WithDDC)' == true"/>
 		<!-- Feature packs -->
 		<Property Name="InstalledRequirements" Value="$(InstalledRequirements);Make Feature Packs"/>
@@ -697,79 +677,83 @@
 
 			<!-- Tag files that always need to be part of an installed build -->
 			<Tag Files="#Compile UnrealHeaderTool Win64" With="#Installed Win64"/>
-			<Tag Files="#UE4Editor Win64;#UE4Editor Win64 Includes" Except="#UE4Editor Win64 Unstripped;#UE4Editor Win64 Unsigned" With="#Installed Win64"/>
+			<Tag Files="#UE4Editor Win64" Except="#UE4Editor Win64 Unstripped;#UE4Editor Win64 Unsigned" With="#Installed Win64"/>
 			<Tag Files="#UE4Editor Win64 Stripped;#UE4Editor Win64 Signed" With="#Saved Output"/>
 			<Tag Files="#Build Tools Win64;#Make Feature Packs" With="#Installed Win64"/>
 
 			<!-- Tag enabled target platforms -->
 			<Do If="'$(WithWin64)' == true">
-				<Tag Files="#UE4Game Win64;#UE4Game Win64 Includes" Except="#UE4Game Win64 Unstripped;#UE4Game Win64 Unsigned" With="#Installed Win64"/>
+				<Tag Files="#UE4Game Win64" Except="#UE4Game Win64 Unstripped;#UE4Game Win64 Unsigned" With="#Installed Win64"/>
 				<Tag Files="#UE4Game Win64 Stripped;#UE4Game Win64 Signed" With="#Saved Output"/>
 				<Property Name="CopyInstalledFilter" Value="$(CopyInstalledFilter);$(CopyWin64Filter)"/>
 				<Property Name="CopyInstalledExceptions" Value="$(CopyInstalledExceptions);$(CopyWin64Exceptions)"/>
 			</Do>
 			<Do If="'$(WithWin32)' == true">
-				<Tag Files="#UE4Game Win32;#UE4Game Win32 Includes" Except="#UE4Game Win32 Unstripped;#UE4Game Win32 Unsigned" With="#Installed Win64"/>
+				<Tag Files="#UE4Game Win32" Except="#UE4Game Win32 Unstripped;#UE4Game Win32 Unsigned" With="#Installed Win64"/>
 				<Tag Files="#UE4Game Win32 Stripped;#UE4Game Win32 Signed" With="#Saved Output"/>
 				<Property Name="CopyInstalledFilter" Value="$(CopyInstalledFilter);$(CopyWin32Filter)"/>
 				<Property Name="CopyInstalledExceptions" Value="$(CopyInstalledExceptions);$(CopyWin32Exceptions)"/>
 			</Do>
 			<Do If="'$(WithAndroid)' == true">
-				<Tag Files="#UE4Game Android;#UE4Game Android Includes" Except="#UE4Game Android Unstripped" With="#Installed Win64"/>
+				<Tag Files="#UE4Game Android" Except="#UE4Game Android Unstripped" With="#Installed Win64"/>
 				<Tag Files="#UE4Game Android Stripped" With="#Saved Output"/>
 				<Property Name="CopyInstalledFilter" Value="$(CopyInstalledFilter);$(CopyAndroidFilter)"/>
 				<Property Name="CopyInstalledExceptions" Value="$(CopyInstalledExceptions);$(CopyAndroidExceptions)"/>
 			</Do>
 			<Do If="'$(WithIOS)' == true">
-				<Tag Files="#UE4Game IOS;#UE4Game IOS Includes" Except="#UE4Game IOS Unstripped" With="#Installed Win64"/>
+				<Tag Files="#UE4Game IOS" Except="#UE4Game IOS Unstripped" With="#Installed Win64"/>
 				<Tag Files="#UE4Game IOS Stripped" With="#Saved Output"/>
 				<Property Name="CopyInstalledFilter" Value="$(CopyInstalledFilter);$(CopyIOSFilter)"/>
 				<Property Name="CopyInstalledExceptions" Value="$(CopyInstalledExceptions);$(CopyIOSExceptions)"/>
 			</Do>
 			<Do If="'$(WithTVOS)' == true">
-				<Tag Files="#UE4Game TVOS;#UE4Game TVOS Includes" Except="#UE4Game TVOS Unstripped" With="#Installed Win64"/>
+				<Tag Files="#UE4Game TVOS" Except="#UE4Game TVOS Unstripped" With="#Installed Win64"/>
 				<Tag Files="#UE4Game TVOS Stripped" With="#Saved Output"/>
 				<Property Name="CopyInstalledFilter" Value="$(CopyInstalledFilter);$(CopyTVOSFilter)"/>
 				<Property Name="CopyInstalledExceptions" Value="$(CopyInstalledExceptions);$(CopyTVOSExceptions)"/>
 			</Do>
 			<Do If="'$(WithLinux)' == true">
-				<Tag Files="#UE4Game Linux;#UE4Game Linux Includes;#Build Tools Linux" Except="#UE4Game Linux Unstripped" With="#Installed Win64"/>
+				<Tag Files="#UE4Game Linux;#Build Tools Linux" Except="#UE4Game Linux Unstripped" With="#Installed Win64"/>
 				<Tag Files="#UE4Game Linux Stripped" With="#Saved Output"/>
 				<Property Name="CopyInstalledFilter" Value="$(CopyInstalledFilter);$(CopyLinuxFilter)"/>
 				<Property Name="CopyInstalledExceptions" Value="$(CopyInstalledExceptions);$(CopyLinuxExceptions)"/>
 			</Do>
 			<Do If="'$(WithHTML5)' == true">
-				<Tag Files="#UE4Game HTML5;#UE4Game HTML5 Includes" With="#Installed Win64"/>
+				<Tag Files="#UE4Game HTML5" With="#Installed Win64"/>
 				<Property Name="CopyInstalledFilter" Value="$(CopyInstalledFilter);$(CopyHTML5Filter)"/>
 				<Property Name="CopyInstalledExceptions" Value="$(CopyInstalledExceptions);$(CopyHTML5Exceptions)"/>
 			</Do>
 			<Do If="'$(WithPS4)' == true">
-				<Tag Files="#UE4Game PS4;#UE4Game PS4 Includes" With="#Installed Win64"/>
+				<Tag Files="#UE4Game PS4" With="#Installed Win64"/>
 				<Property Name="CopyInstalledFilter" Value="$(CopyInstalledFilter);$(CopyPS4Filter)"/>
 				<Property Name="CopyInstalledExceptions" Value="$(CopyInstalledExceptions);$(CopyPS4Exceptions)"/>
 			</Do>
 			<Do If="'$(WithXboxOne)' == true">
-				<Tag Files="#UE4Game XboxOne;#UE4Game XboxOne Includes" With="#Installed Win64"/>
+				<Tag Files="#UE4Game XboxOne" With="#Installed Win64"/>
 				<Property Name="CopyInstalledFilter" Value="$(CopyInstalledFilter);$(CopyXboxOneFilter)"/>
 				<Property Name="CopyInstalledExceptions" Value="$(CopyInstalledExceptions);$(CopyXboxOneExceptions)"/>
 			</Do>
 			<Do If="'$(WithSwitch)' == true">
-				<Tag Files="#UE4Game Switch;#UE4Game Switch Includes" With="#Installed Win64"/>
+				<Tag Files="#UE4Game Switch" With="#Installed Win64"/>
 				<Property Name="CopyInstalledFilter" Value="$(CopyInstalledFilter);$(CopySwitchFilter)"/>
 				<Property Name="CopyInstalledExceptions" Value="$(CopyInstalledExceptions);$(CopySwitchExceptions)"/>
 			</Do>
 			<Do If="'$(WithLumin)' == true">
-				<Tag Files="#UE4Game Lumin;#UE4Game Lumin Includes" With="#Installed Win64"/>
+				<Tag Files="#UE4Game Lumin" With="#Installed Win64"/>
 				<Property Name="CopyInstalledFilter" Value="$(CopyInstalledFilter);$(CopyLuminFilter)"/>
 				<Property Name="CopyInstalledExceptions" Value="$(CopyInstalledExceptions);$(CopyLuminExceptions)"/>
 			</Do>
 
 			<!-- Tag any dependencies from all previous build products -->
 			<Tag Files="#Installed Win64" Filter="*.target" With="#TargetReceipts"/>
-			<TagReceipt Files="#TargetReceipts" RuntimeDependencies="true" PrecompiledBuildDependencies="true" PrecompiledRuntimeDependencies="true" With="#Installed Win64"/>
+			<TagReceipt Files="#TargetReceipts" RuntimeDependencies="true" With="#Installed Win64"/>
+
+			<!-- Tag any additional dependencies -->
+			<Tag Files="#Installed Win64" Filter=".../DependencyList.txt;.../DependencyList-AllModules.txt" With="#DependencyLists"/>
+			<Tag FileLists="#DependencyLists" With="#Installed Win64" />
 
 			<!-- Confidential Exceptions - done last to ensure they never get overwritten -->
-			<Property Name="CopyInstalledExceptions" Value="$(CopyInstalledExceptions);$(ConfidentialExceptions)"/>
+			<Property Name="CopyInstalledExceptions" Value="$(CopyInstalledExceptions);#DependencyLists;$(ConfidentialExceptions)"/>
 			
 			<!-- Tag additional files needed for installed build -->
 			<Tag Files="$(CopyInstalledFilter)" With="#Installed Win64"/>
@@ -819,15 +803,15 @@
 
 	<Agent Name="Installed Build Group Mac" Type="Mac">
 		<!-- Build up a list of requirements as some can be optional -->
-		<Property Name="FilterRequirements" Value="#Compile UnrealHeaderTool Mac;#UE4Editor Mac;#UE4Editor Mac Unstripped;#UE4Editor Mac Stripped;#UE4Editor Mac Unsigned;#UE4Editor Mac Signed;#UE4Editor Mac Includes;#Build Tools Mac;#Build Tools CS Binaries"/>
+		<Property Name="FilterRequirements" Value="#Compile UnrealHeaderTool Mac;#UE4Editor Mac;#UE4Editor Mac Unstripped;#UE4Editor Mac Stripped;#UE4Editor Mac Unsigned;#UE4Editor Mac Signed;#Build Tools Mac;#Build Tools CS Binaries"/>
 
 		<!-- Optional Target Platform requirements -->
-		<Property Name="FilterRequirements" Value="$(FilterRequirements);#UE4Game Mac;#UE4Game Mac Unstripped;#UE4Game Mac Stripped;#UE4Game Mac Unsigned;#UE4Game Mac Signed;#UE4Game Mac Includes" If="'$(WithMac)' == true"/>
-		<Property Name="FilterRequirements" Value="$(FilterRequirements);#UE4Game Android;#UE4Game Android Unstripped;#UE4Game Android Stripped;#UE4Game Android Includes" If="'$(WithAndroid)' == true"/>
-		<Property Name="FilterRequirements" Value="$(FilterRequirements);#UE4Game IOS;#UE4Game IOS Unstripped;#UE4Game IOS Stripped;#UE4Game IOS Includes" If="'$(WithIOS)' == true"/>
-		<Property Name="FilterRequirements" Value="$(FilterRequirements);#UE4Game TVOS;#UE4Game TVOS Unstripped;#UE4Game TVOS Stripped;#UE4Game TVOS Includes" If="'$(WithTVOS)' == true"/>
-		<Property Name="FilterRequirements" Value="$(FilterRequirements);#UE4Game HTML5;#UE4Game HTML5 Includes" If="'$(WithHTML5)' == true"/>
-		<Property Name="FilterRequirements" Value="$(FilterRequirements);#UE4Game LuminMac;#UE4Game LuminMac Includes;#UE4Game LuminMac Unstripped;#UE4Game LuminMac Stripped;#UE4Game LuminMac Includes" If="'$(WithLumin)' == true"/> 
+		<Property Name="FilterRequirements" Value="$(FilterRequirements);#UE4Game Mac;#UE4Game Mac Unstripped;#UE4Game Mac Stripped;#UE4Game Mac Unsigned;#UE4Game Mac Signed" If="'$(WithMac)' == true"/>
+		<Property Name="FilterRequirements" Value="$(FilterRequirements);#UE4Game Android;#UE4Game Android Unstripped;#UE4Game Android Stripped" If="'$(WithAndroid)' == true"/>
+		<Property Name="FilterRequirements" Value="$(FilterRequirements);#UE4Game IOS;#UE4Game IOS Unstripped;#UE4Game IOS Stripped" If="'$(WithIOS)' == true"/>
+		<Property Name="FilterRequirements" Value="$(FilterRequirements);#UE4Game TVOS;#UE4Game TVOS Unstripped;#UE4Game TVOS Stripped" If="'$(WithTVOS)' == true"/>
+		<Property Name="FilterRequirements" Value="$(FilterRequirements);#UE4Game HTML5" If="'$(WithHTML5)' == true"/>
+		<Property Name="FilterRequirements" Value="$(FilterRequirements);#UE4Game LuminMac;#UE4Game LuminMac Unstripped;#UE4Game LuminMac Stripped" If="'$(WithLumin)' == true"/> 
 		<Property Name="FilterRequirements" Value="$(FilterRequirements);#CompressedDDCMac" If="'$(WithDDC)' == true"/>
 
 		<!-- Feature packs -->
@@ -844,42 +828,42 @@
 
 			<!-- Tag files that always need to be part of an installed build -->
 			<Tag Files="#Compile UnrealHeaderTool Mac" With="#Installed Mac"/>
-			<Tag Files="#UE4Editor Mac;#UE4Editor Mac Includes" Except="#UE4Editor Mac Unstripped;#UE4Editor Mac Unsigned" With="#Installed Mac"/>
+			<Tag Files="#UE4Editor Mac" Except="#UE4Editor Mac Unstripped;#UE4Editor Mac Unsigned" With="#Installed Mac"/>
 			<Tag Files="#UE4Editor Mac Stripped;#UE4Editor Mac Signed" With="#Saved Output"/>
 			<Tag Files="#Build Tools Mac;#Make Feature Packs" With="#Installed Mac"/>
 
 			<!-- Tag enabled target platforms -->
 			<Do If="'$(WithMac)' == true">
-				<Tag Files="#UE4Game Mac;#UE4Game Mac Includes" Except="#UE4Game Mac Unstripped;#UE4Game Mac Unsigned" With="#Installed Mac"/>
+				<Tag Files="#UE4Game Mac" Except="#UE4Game Mac Unstripped;#UE4Game Mac Unsigned" With="#Installed Mac"/>
 				<Tag Files="#UE4Game Mac Stripped;#UE4Game Mac Signed" With="#Saved Output"/>
 				<Property Name="CopyInstalledFilter" Value="$(CopyInstalledFilter);$(CopyMacFilter)"/>
 				<Property Name="CopyInstalledExceptions" Value="$(CopyInstalledExceptions);$(CopyMacExceptions)"/>
 			</Do>
 			<Do  If="'$(WithAndroid)' == true">
-				<Tag Files="#UE4Game Android;#UE4Game Android Includes" Except="#UE4Game Android Unstripped" With="#Installed Mac"/>
+				<Tag Files="#UE4Game Android" Except="#UE4Game Android Unstripped" With="#Installed Mac"/>
 				<Tag Files="#UE4Game Android Stripped" With="#Saved Output"/>
 				<Property Name="CopyInstalledFilter" Value="$(CopyInstalledFilter);$(CopyAndroidFilter)"/>
 				<Property Name="CopyInstalledExceptions" Value="$(CopyInstalledExceptions);$(CopyAndroidExceptions)"/>
 			</Do>
 			<Do If="'$(WithIOS)' == true">
-				<Tag Files="#UE4Game IOS;#UE4Game IOS Includes" Except="#UE4Game IOS Unstripped" With="#Installed Mac"/>
+				<Tag Files="#UE4Game IOS" Except="#UE4Game IOS Unstripped" With="#Installed Mac"/>
 				<Tag Files="#UE4Game IOS Stripped" With="#Saved Output"/>
 				<Property Name="CopyInstalledFilter" Value="$(CopyInstalledFilter);$(CopyIOSFilter)"/>
 				<Property Name="CopyInstalledExceptions" Value="$(CopyInstalledExceptions);$(CopyIOSExceptions)"/>
 			</Do>
 			<Do If="'$(WithTVOS)' == true">
-				<Tag Files="#UE4Game TVOS;#UE4Game TVOS Includes" Except="#UE4Game TVOS Unstripped" With="#Installed Mac"/>
+				<Tag Files="#UE4Game TVOS" Except="#UE4Game TVOS Unstripped" With="#Installed Mac"/>
 				<Tag Files="#UE4Game TVOS Stripped" With="#Saved Output"/>
 				<Property Name="CopyInstalledFilter" Value="$(CopyInstalledFilter);$(CopyTVOSFilter)"/>
 				<Property Name="CopyInstalledExceptions" Value="$(CopyInstalledExceptions);$(CopyTVOSExceptions)"/>
 			</Do>
 			<Do If="'$(WithHTML5)' == true">
-				<Tag Files="#UE4Game HTML5;#UE4Game HTML5 Includes" With="#Installed Mac"/>
+				<Tag Files="#UE4Game HTML5" With="#Installed Mac"/>
 				<Property Name="CopyInstalledFilter" Value="$(CopyInstalledFilter);$(CopyHTML5Filter)"/>
 				<Property Name="CopyInstalledExceptions" Value="$(CopyInstalledExceptions);$(CopyHTML5Exceptions)"/>
 			</Do>
 			<Do If="'$(WithLumin)' == true">
-				<Tag Files="#UE4Game LuminMac;#UE4Game LuminMac Includes" Except="#UE4Game LuminMac Unstripped" With="#Installed Mac"/>
+				<Tag Files="#UE4Game LuminMac" Except="#UE4Game LuminMac Unstripped" With="#Installed Mac"/>
 				<Tag Files="#UE4Game LuminMac Stripped" With="#Saved Output"/>
 				<Property Name="CopyInstalledFilter" Value="$(CopyInstalledFilter);$(CopyLuminFilter)"/>
 				<Property Name="CopyInstalledExceptions" Value="$(CopyInstalledExceptions);$(CopyLuminExceptions)"/>
@@ -887,10 +871,14 @@
 
 			<!-- Tag any dependencies from all previous build products -->
 			<Tag Files="#Installed Mac" Filter="*.target" With="#TargetReceipts"/>
-			<TagReceipt Files="#TargetReceipts" RuntimeDependencies="true" PrecompiledBuildDependencies="true" PrecompiledRuntimeDependencies="true" With="#Installed Mac"/>
+			<TagReceipt Files="#TargetReceipts" RuntimeDependencies="true" With="#Installed Mac"/>
+
+			<!-- Tag any additional dependencies -->
+			<Tag Files="#Installed Mac" Filter=".../DependencyList.txt" With="#DependencyLists"/>
+			<Tag FileLists="#DependencyLists" With="#Installed Mac" />
 
 			<!-- Confidential Exceptions - done last to ensure they never get overwritten -->
-			<Property Name="CopyInstalledExceptions" Value="$(CopyInstalledExceptions);$(ConfidentialExceptions)"/>
+			<Property Name="CopyInstalledExceptions" Value="$(CopyInstalledExceptions);$(ConfidentialExceptions);#DependencyLists"/>
 			
 			<!-- Tag additional files needed for editor and platforms -->
 			<Tag Files="$(CopyInstalledFilter)" With="#Installed Mac"/>
@@ -936,10 +924,10 @@
 
 	<Agent Name="Installed Build Group Linux" Type="Win64">
 		<!-- Build up a list of requirements as some can be optional -->
-		<Property Name="InstalledRequirements" Value="#Compile UnrealHeaderTool Linux;#UE4Editor Linux;#UE4Editor Linux Unstripped;#UE4Editor Linux Stripped;#UE4Editor Linux Includes;#Build Tools Linux;#Build Tools CS Binaries"/>
+		<Property Name="InstalledRequirements" Value="#Compile UnrealHeaderTool Linux;#UE4Editor Linux;#UE4Editor Linux Unstripped;#UE4Editor Linux Stripped;#Build Tools Linux;#Build Tools CS Binaries"/>
 
 		<!-- Optional Target Platform requirements -->
-		<Property Name="InstalledRequirements" Value="$(InstalledRequirements);#UE4Game Linux;#UE4Game Linux Unstripped;#UE4Game Linux Stripped;#UE4Game Linux Includes" If="'$(WithLinux)' == true"/>
+		<Property Name="InstalledRequirements" Value="$(InstalledRequirements);#UE4Game Linux;#UE4Game Linux Unstripped;#UE4Game Linux Stripped" If="'$(WithLinux)' == true"/>
 		<Property Name="InstalledRequirements" Value="$(InstalledRequirements);#CompressedDDCLinux" If="'$(WithDDC)' == true"/>
 
 		<!-- Feature packs -->
@@ -957,13 +945,13 @@
 
 			<!-- Tag files that always need to be part of an installed build -->
 			<Tag Files="#Compile UnrealHeaderTool Linux" With="#Installed Linux"/>
-			<Tag Files="#UE4Editor Linux;#UE4Editor Linux Includes" Except="#UE4Editor Linux Unstripped" With="#Installed Linux"/>
+			<Tag Files="#UE4Editor Linux" Except="#UE4Editor Linux Unstripped" With="#Installed Linux"/>
 			<Tag Files="#UE4Editor Linux Stripped" With="#Saved Output"/>
 			<Tag Files="#Build Tools Linux;#Make Feature Packs" With="#Installed Linux"/>
 
 			<!-- Tag enabled target platforms -->
 			<Do If="'$(WithLinux)' == true">
-				<Tag Files="#UE4Game Linux;#UE4Game Linux Includes;#Build Tools Linux" Except="#UE4Game Linux Unstripped" With="#Installed Linux"/>
+				<Tag Files="#UE4Game Linux;#Build Tools Linux" Except="#UE4Game Linux Unstripped" With="#Installed Linux"/>
 				<Tag Files="#UE4Game Linux Stripped" With="#Saved Output"/>
 				<Property Name="CopyInstalledFilter" Value="$(CopyInstalledFilter);$(CopyLinuxFilter)"/>
 				<Property Name="CopyInstalledExceptions" Value="$(CopyInstalledExceptions);$(CopyLinuxExceptions)"/>
@@ -971,10 +959,14 @@
 
 			<!-- Tag any dependencies from all previous build products -->
 			<Tag Files="#Installed Linux" Filter="*.target" With="#TargetReceipts"/>
-			<TagReceipt Files="#TargetReceipts" RuntimeDependencies="true" PrecompiledBuildDependencies="true" PrecompiledRuntimeDependencies="true" With="#Installed Linux"/>
+			<TagReceipt Files="#TargetReceipts" RuntimeDependencies="true" With="#Installed Linux"/>
+
+			<!-- Tag any additional dependencies -->
+			<Tag Files="#Installed Linux" Filter=".../DependencyList.txt" With="#DependencyLists"/>
+			<Tag FileLists="#DependencyLists" With="#Installed Linux" />
 
 			<!-- Confidential Exceptions - done last to ensure they never get overwritten -->
-			<Property Name="CopyInstalledExceptions" Value="$(CopyInstalledExceptions);$(ConfidentialExceptions)"/>
+			<Property Name="CopyInstalledExceptions" Value="$(CopyInstalledExceptions);$(ConfidentialExceptions);#DependencyLists"/>
 
 			<!-- Tag additional files needed for installed build -->
 			<Tag Files="$(CopyInstalledFilter)" With="#Installed Linux"/>
diff --git a/Engine/Source/Programs/AutomationTool/AutomationUtils/UE4Build.cs b/Engine/Source/Programs/AutomationTool/AutomationUtils/UE4Build.cs
--- a/Engine/Source/Programs/AutomationTool/AutomationUtils/UE4Build.cs
+++ b/Engine/Source/Programs/AutomationTool/AutomationUtils/UE4Build.cs
@@ -95,10 +95,6 @@ namespace AutomationTool
 					PrepareBuildProduct(Item);
 				}
 			}
-			foreach(string Item in Manifest.LibraryBuildProducts)
-			{
-				LibraryBuildProductFiles.Add(Item);
-			}
 			return Manifest;
 		}
 
@@ -125,10 +121,6 @@ namespace AutomationTool
 				}
 				AddBuildProduct(Item);
 			}
-			foreach(string Item in Manifest.LibraryBuildProducts)
-			{
-				LibraryBuildProductFiles.Add(Item);
-			}
 		}
 		
 		private void PrepareUBT()
@@ -710,7 +702,6 @@ namespace AutomationTool
 		{
 			OwnerCommand = Command;
 			BuildProductFiles.Clear();
-			LibraryBuildProductFiles.Clear();
 		}
 
 		public List<string> FindXGEFiles()
@@ -1672,9 +1663,6 @@ namespace AutomationTool
 		// List of everything we built so far
 		public readonly List<string> BuildProductFiles = new List<string>();
 
-		// Library files that were part of the build
-		public readonly List<string> LibraryBuildProductFiles = new List<string>();
-
 		private bool DeleteBuildProducts = true;
 	}
 }
diff --git a/Engine/Source/Programs/AutomationTool/BuildGraph/Tasks/CompileTask.cs b/Engine/Source/Programs/AutomationTool/BuildGraph/Tasks/CompileTask.cs
--- a/Engine/Source/Programs/AutomationTool/BuildGraph/Tasks/CompileTask.cs
+++ b/Engine/Source/Programs/AutomationTool/BuildGraph/Tasks/CompileTask.cs
@@ -1,4 +1,4 @@
-﻿// Copyright 1998-2018 Epic Games, Inc. All Rights Reserved.
+// Copyright 1998-2018 Epic Games, Inc. All Rights Reserved.
 
 using System;
 using System.Collections.Generic;
@@ -168,13 +168,11 @@ namespace AutomationTool
 				{
 					HashSet<FileReference> FileSet = CustomTask.FindOrAddTagSet(TagNameToFileSet, TagName);
 					FileSet.UnionWith(Manifest.BuildProducts.Select(x => new FileReference(x)));
-					FileSet.UnionWith(Manifest.LibraryBuildProducts.Select(x => new FileReference(x)));
 				}
 			}
 
 			// Add everything to the list of build products
 			BuildProducts.UnionWith(Builder.BuildProductFiles.Select(x => new FileReference(x)));
-			BuildProducts.UnionWith(Builder.LibraryBuildProductFiles.Select(x => new FileReference(x)));
 		}
 	}
 
diff --git a/Engine/Source/Programs/AutomationTool/BuildGraph/Tasks/StageTask.cs b/Engine/Source/Programs/AutomationTool/BuildGraph/Tasks/StageTask.cs
--- a/Engine/Source/Programs/AutomationTool/BuildGraph/Tasks/StageTask.cs
+++ b/Engine/Source/Programs/AutomationTool/BuildGraph/Tasks/StageTask.cs
@@ -1,4 +1,4 @@
-﻿// Copyright 1998-2018 Epic Games, Inc. All Rights Reserved.
+// Copyright 1998-2018 Epic Games, Inc. All Rights Reserved.
 
 using AutomationTool;
 using System;
@@ -126,7 +126,7 @@ namespace BuildGraph.Tasks
 
 			// Stage all the build products needed at runtime
 			HashSet<FileReference> SourceFiles = new HashSet<FileReference>();
-			foreach(BuildProduct BuildProduct in Receipt.BuildProducts.Where(x => x.Type != BuildProductType.StaticLibrary && x.Type != BuildProductType.ImportLibrary))
+			foreach(BuildProduct BuildProduct in Receipt.BuildProducts)
 			{
 				SourceFiles.Add(BuildProduct.Path);
 			}
diff --git a/Engine/Source/Programs/AutomationTool/BuildGraph/Tasks/TagReceiptTask.cs b/Engine/Source/Programs/AutomationTool/BuildGraph/Tasks/TagReceiptTask.cs
--- a/Engine/Source/Programs/AutomationTool/BuildGraph/Tasks/TagReceiptTask.cs
+++ b/Engine/Source/Programs/AutomationTool/BuildGraph/Tasks/TagReceiptTask.cs
@@ -59,18 +59,6 @@ namespace AutomationTool.Tasks
 		[TaskParameter(Optional = true)]
 		public string StagedFileType;
 
-		/// <summary>
-		/// Whether to tag the Precompiled Build Dependencies listed in receipts
-		/// </summary>
-		[TaskParameter(Optional = true)]
-		public bool PrecompiledBuildDependencies;
-
-		/// <summary>
-		/// Whether to tag the Precompiled Runtime Dependencies listed in receipts
-		/// </summary>
-		[TaskParameter(Optional = true)]
-		public bool PrecompiledRuntimeDependencies;
-
 		/// <summary>
 		/// Name of the tag to apply
 		/// </summary>
@@ -184,40 +172,6 @@ namespace AutomationTool.Tasks
 						}
 					}
 				}
-
-				if (Parameters.PrecompiledBuildDependencies)
-				{
-					foreach(FileReference PrecompiledBuildDependency in Receipt.PrecompiledBuildDependencies)
-					{
-						// Only add files that exist as dependencies are assumed to always exist
-						FileReference DependencyPath = PrecompiledBuildDependency;
-						if (FileReference.Exists(DependencyPath))
-						{
-							Files.Add(DependencyPath);
-						}
-						else
-						{
-							CommandUtils.LogWarning("File listed as PrecompiledBuildDependency in {0} does not exist ({1})", TargetFile.FullName, DependencyPath.FullName);
-						}
-					}
-				}
-
-				if (Parameters.PrecompiledRuntimeDependencies)
-				{
-					foreach (FileReference PrecompiledRuntimeDependency in Receipt.PrecompiledRuntimeDependencies)
-					{
-						// Only add files that exist as dependencies are assumed to always exist
-						FileReference DependencyPath = PrecompiledRuntimeDependency;
-						if (FileReference.Exists(DependencyPath))
-						{
-							Files.Add(DependencyPath);
-						}
-						else
-						{
-							CommandUtils.LogWarning("File listed as PrecompiledRuntimeDependency in {0} does not exist ({1})", TargetFile.FullName, DependencyPath.FullName);
-						}
-					}
-				}
 			}
 
 			// Apply the tag to all the matching files
diff --git a/Engine/Source/Programs/AutomationTool/BuildGraph/Tasks/TagTask.cs b/Engine/Source/Programs/AutomationTool/BuildGraph/Tasks/TagTask.cs
--- a/Engine/Source/Programs/AutomationTool/BuildGraph/Tasks/TagTask.cs
+++ b/Engine/Source/Programs/AutomationTool/BuildGraph/Tasks/TagTask.cs
@@ -1,4 +1,4 @@
-﻿// Copyright 1998-2018 Epic Games, Inc. All Rights Reserved.
+// Copyright 1998-2018 Epic Games, Inc. All Rights Reserved.
 
 using System;
 using System.Collections.Generic;
@@ -26,9 +26,15 @@ namespace AutomationTool.Tasks
 		/// <summary>
 		/// Set of files to work from, including wildcards and tag names, separated by semicolons. Resolved relative to BaseDir if set, otherwise to the branch root directory.
 		/// </summary>
-		[TaskParameter(ValidationType = TaskParameterValidationType.FileSpec)]
+		[TaskParameter(Optional = true, ValidationType = TaskParameterValidationType.FileSpec)]
 		public string Files;
 
+		/// <summary>
+		/// Set of text files to add additional files from. Each file list should have on file per line.
+		/// </summary>
+		[TaskParameter(Optional = true, ValidationType = TaskParameterValidationType.FileSpec)]
+		public string FileLists;
+
 		/// <summary>
 		/// Patterns to filter the list of files by, including tag names or wildcards. May include patterns that apply to the base directory if set. Defaults to all files if not specified.
 		/// </summary>
@@ -84,7 +90,34 @@ namespace AutomationTool.Tasks
 			List<string> ExcludeRules = ParseRules(BaseDir, Parameters.Except ?? "", TagNameToFileSet);
 
 			// Resolve the input list
-			HashSet<FileReference> Files = ResolveFilespecWithExcludePatterns(BaseDir, Parameters.Files, ExcludeRules, TagNameToFileSet);
+			HashSet<FileReference> Files = new HashSet<FileReference>();
+			if(!String.IsNullOrEmpty(Parameters.Files))
+			{
+				Files.UnionWith(ResolveFilespecWithExcludePatterns(BaseDir, Parameters.Files, ExcludeRules, TagNameToFileSet));
+			}
+
+			// Resolve the input file lists
+			if(!String.IsNullOrEmpty(Parameters.FileLists))
+			{
+				HashSet<FileReference> FileLists = ResolveFilespec(BaseDir, Parameters.FileLists, TagNameToFileSet);
+				foreach(FileReference FileList in FileLists)
+				{
+					if(!FileReference.Exists(FileList))
+					{
+						throw new BuildException("Specified file list '{0}' does not exist", FileList);
+					}
+
+					string[] Lines = FileReference.ReadAllLines(FileList);
+					foreach(string Line in Lines)
+					{
+						string TrimLine = Line.Trim();
+						if(TrimLine.Length > 0)
+						{
+							Files.Add(FileReference.Combine(BaseDir, TrimLine));
+						}
+					}
+				}
+			}
 
 			// Limit to matches against the 'Filter' parameter, if set
 			if(Parameters.Filter != null)
diff --git a/Engine/Source/Programs/AutomationTool/Scripts/BuildPluginCommand.Automation.cs b/Engine/Source/Programs/AutomationTool/Scripts/BuildPluginCommand.Automation.cs
--- a/Engine/Source/Programs/AutomationTool/Scripts/BuildPluginCommand.Automation.cs
+++ b/Engine/Source/Programs/AutomationTool/Scripts/BuildPluginCommand.Automation.cs
@@ -1,4 +1,4 @@
-﻿// Copyright 1998-2018 Epic Games, Inc. All Rights Reserved.
+// Copyright 1998-2018 Epic Games, Inc. All Rights Reserved.
 using System;
 using System.Collections.Generic;
 using System.IO;
@@ -157,8 +157,6 @@ class BuildPlugin : BuildCommand
 				throw new AutomationException("Missing or invalid target receipt ({0})", ReceiptFileName);
 			}
 			BuildProducts.UnionWith(Receipt.BuildProducts.Select(x => x.Path).Where(x => x.IsUnderDirectory(HostProjectPluginFile.Directory)));
-			BuildProducts.UnionWith(Receipt.PrecompiledBuildDependencies.Where(x => x.IsUnderDirectory(HostProjectPluginFile.Directory)));
-			BuildProducts.UnionWith(Receipt.PrecompiledRuntimeDependencies.Where(x => x.IsUnderDirectory(HostProjectPluginFile.Directory)));
 		}
 		return BuildProducts.ToArray();
 	}
diff --git a/Engine/Source/Programs/UnrealBuildTool/Configuration/TargetRules.cs b/Engine/Source/Programs/UnrealBuildTool/Configuration/TargetRules.cs
--- a/Engine/Source/Programs/UnrealBuildTool/Configuration/TargetRules.cs
+++ b/Engine/Source/Programs/UnrealBuildTool/Configuration/TargetRules.cs
@@ -321,6 +321,12 @@ namespace UnrealBuildTool
 		[Obsolete("bBuildAllPlugins has been deprecated. Use bPrecompile to build all modules which are not part of the target.")]
 		public bool bBuildAllPlugins = false;
 
+		/// <summary>
+		/// Build all the modules that are valid for this target type. Used for CIS and making installed engine builds.
+		/// </summary>
+		[CommandLine("-AllModules")]
+		public bool bBuildAllModules = false;
+
 		/// <summary>
 		/// A list of additional plugins which need to be included in this target. This allows referencing non-optional plugin modules
 		/// which cannot be disabled, and allows building against specific modules in program targets which do not fit the categories
@@ -1096,6 +1102,12 @@ namespace UnrealBuildTool
 		[CommandLine("-Manifest")]
 		public List<FileReference> ManifestFileNames = new List<FileReference>();
 
+		/// <summary>
+		/// Path to a list of dependencies for this target, when precompiling
+		/// </summary>
+		[CommandLine("-DependencyList")]
+		public List<FileReference> DependencyListFileNames = new List<FileReference>();
+
 		/// <summary>
 		/// Specifies the build environment for this target. See TargetBuildEnvironment for more information on the available options.
 		/// </summary>
@@ -1515,6 +1527,11 @@ namespace UnrealBuildTool
 			get { return Inner.bBuildAllPlugins; }
 		}
 
+		public bool bBuildAllModules
+		{
+			get { return Inner.bBuildAllModules; }
+		}
+
 		public IEnumerable<string> AdditionalPlugins
 		{
 			get { return Inner.AdditionalPlugins; }
@@ -2075,6 +2092,11 @@ namespace UnrealBuildTool
 			get { return Inner.ManifestFileNames.AsReadOnly(); }
 		}
 
+		public IReadOnlyList<FileReference> DependencyListFileNames
+		{
+			get { return Inner.DependencyListFileNames.AsReadOnly(); }
+		}
+
 		public TargetBuildEnvironment BuildEnvironment
 		{
 			get { return Inner.BuildEnvironment; }
diff --git a/Engine/Source/Programs/UnrealBuildTool/Configuration/UEBuildBinary.cs b/Engine/Source/Programs/UnrealBuildTool/Configuration/UEBuildBinary.cs
--- a/Engine/Source/Programs/UnrealBuildTool/Configuration/UEBuildBinary.cs
+++ b/Engine/Source/Programs/UnrealBuildTool/Configuration/UEBuildBinary.cs
@@ -404,44 +404,80 @@ namespace UnrealBuildTool
 		/// <param name="bCreateDebugInfo">Whether debug info is enabled for this binary</param>
 		public void GetBuildProducts(ReadOnlyTargetRules Target, UEToolChain ToolChain, Dictionary<FileReference, BuildProductType> BuildProducts, bool bCreateDebugInfo)
 		{
-			// Get the type of build products we're creating
-			BuildProductType OutputType = BuildProductType.RequiredResource;
-			switch (Type)
+			// Add all the precompiled outputs
+			if(Target.bPrecompile)
 			{
-				case UEBuildBinaryType.Executable:
-					OutputType = BuildProductType.Executable;
-					break;
-				case UEBuildBinaryType.DynamicLinkLibrary:
-					OutputType = BuildProductType.DynamicLibrary;
-					break;
-				case UEBuildBinaryType.StaticLibrary:
-					OutputType = BuildProductType.StaticLibrary;
-					break;
-			}
+				foreach(UEBuildModuleCPP Module in Modules.OfType<UEBuildModuleCPP>())
+				{
+					if(Module.Rules.bPrecompile)
+					{
+						if(Module.GeneratedCodeDirectory != null && DirectoryReference.Exists(Module.GeneratedCodeDirectory))
+						{
+							foreach(FileReference GeneratedCodeFile in DirectoryReference.EnumerateFiles(Module.GeneratedCodeDirectory))
+							{
+								BuildProducts.Add(GeneratedCodeFile, BuildProductType.BuildResource);
+							}
+						}
+						if(Target.LinkType == TargetLinkType.Monolithic)
+						{
+							FileReference PrecompiledManifestLocation = Module.PrecompiledManifestLocation;
+							BuildProducts.Add(PrecompiledManifestLocation, BuildProductType.BuildResource);
 
-			// Add the primary build products
-			string[] DebugExtensions = UEBuildPlatform.GetBuildPlatform(Target.Platform).GetDebugInfoExtensions(Target, Type);
-			foreach (FileReference OutputFilePath in OutputFilePaths)
-			{
-				AddBuildProductAndDebugFiles(OutputFilePath, OutputType, DebugExtensions, BuildProducts, ToolChain, bCreateDebugInfo);
+							PrecompiledManifest ModuleManifest = PrecompiledManifest.Read(PrecompiledManifestLocation);
+							foreach(FileReference OutputFile in ModuleManifest.OutputFiles)
+							{
+								if(!BuildProducts.ContainsKey(OutputFile))
+								{
+									BuildProducts.Add(OutputFile, BuildProductType.BuildResource);
+								}
+							}
+						}
+					}
+				}
 			}
 
-			// Add the console app, if there is one
-			if (Type == UEBuildBinaryType.Executable && bBuildAdditionalConsoleApp)
+			// Add all the binary outputs
+			if(!Target.bDisableLinking)
 			{
+				// Get the type of build products we're creating
+				BuildProductType OutputType = BuildProductType.RequiredResource;
+				switch (Type)
+				{
+					case UEBuildBinaryType.Executable:
+						OutputType = BuildProductType.Executable;
+						break;
+					case UEBuildBinaryType.DynamicLinkLibrary:
+						OutputType = BuildProductType.DynamicLibrary;
+						break;
+					case UEBuildBinaryType.StaticLibrary:
+						OutputType = BuildProductType.BuildResource;
+						break;
+				}
+
+				// Add the primary build products
+				string[] DebugExtensions = UEBuildPlatform.GetBuildPlatform(Target.Platform).GetDebugInfoExtensions(Target, Type);
 				foreach (FileReference OutputFilePath in OutputFilePaths)
 				{
-					AddBuildProductAndDebugFiles(GetAdditionalConsoleAppPath(OutputFilePath), OutputType, DebugExtensions, BuildProducts, ToolChain, bCreateDebugInfo);
+					AddBuildProductAndDebugFiles(OutputFilePath, OutputType, DebugExtensions, BuildProducts, ToolChain, bCreateDebugInfo);
 				}
-			}
 
-			// Add any additional build products from the modules in this binary, including additional bundle resources/dylibs on Mac.
-			List<string> Libraries = new List<string>();
-			List<UEBuildBundleResource> BundleResources = new List<UEBuildBundleResource>();
-			GatherAdditionalResources(Libraries, BundleResources);
+				// Add the console app, if there is one
+				if (Type == UEBuildBinaryType.Executable && bBuildAdditionalConsoleApp)
+				{
+					foreach (FileReference OutputFilePath in OutputFilePaths)
+					{
+						AddBuildProductAndDebugFiles(GetAdditionalConsoleAppPath(OutputFilePath), OutputType, DebugExtensions, BuildProducts, ToolChain, bCreateDebugInfo);
+					}
+				}
+
+				// Add any additional build products from the modules in this binary, including additional bundle resources/dylibs on Mac.
+				List<string> Libraries = new List<string>();
+				List<UEBuildBundleResource> BundleResources = new List<UEBuildBundleResource>();
+				GatherAdditionalResources(Libraries, BundleResources);
 
-			// Add any extra files from the toolchain
-			ToolChain.ModifyBuildProducts(Target, this, Libraries, BundleResources, BuildProducts);
+				// Add any extra files from the toolchain
+				ToolChain.ModifyBuildProducts(Target, this, Libraries, BundleResources, BuildProducts);
+			}
 		}
 
 		/// <summary>
diff --git a/Engine/Source/Programs/UnrealBuildTool/Configuration/UEBuildModuleCPP.cs b/Engine/Source/Programs/UnrealBuildTool/Configuration/UEBuildModuleCPP.cs
--- a/Engine/Source/Programs/UnrealBuildTool/Configuration/UEBuildModuleCPP.cs
+++ b/Engine/Source/Programs/UnrealBuildTool/Configuration/UEBuildModuleCPP.cs
@@ -378,7 +378,7 @@ namespace UnrealBuildTool
 			CppCompileEnvironment ModuleCompileEnvironment = CreateModuleCompileEnvironment(Target, BinaryCompileEnvironment);
 
 			// If the module is precompiled, read the object files from the manifest
-			if(Rules.bUsePrecompiled)
+			if(Rules.bUsePrecompiled && Target.LinkType == TargetLinkType.Monolithic)
 			{
 				PrecompiledManifest Manifest = PrecompiledManifest.Read(PrecompiledManifestLocation);
 				foreach(FileReference OutputFile in Manifest.OutputFiles)
@@ -635,7 +635,7 @@ namespace UnrealBuildTool
 			}
 
 			// Write the compiled manifest
-			if(Rules.bPrecompile)
+			if(Rules.bPrecompile && Target.LinkType == TargetLinkType.Monolithic)
 			{
 				DirectoryReference.CreateDirectory(PrecompiledManifestLocation.Directory);
 
diff --git a/Engine/Source/Programs/UnrealBuildTool/Configuration/UEBuildTarget.cs b/Engine/Source/Programs/UnrealBuildTool/Configuration/UEBuildTarget.cs
--- a/Engine/Source/Programs/UnrealBuildTool/Configuration/UEBuildTarget.cs
+++ b/Engine/Source/Programs/UnrealBuildTool/Configuration/UEBuildTarget.cs
@@ -213,11 +213,6 @@ namespace UnrealBuildTool
 		/// </summary>
 		public readonly List<string> BuildProducts = new List<string>();
 
-		/// <summary>
-		/// 
-		/// </summary>
-		public readonly List<string> LibraryBuildProducts = new List<string>();
-
 		/// <summary>
 		/// 
 		/// </summary>
@@ -261,19 +256,6 @@ namespace UnrealBuildTool
 				AddBuildProduct(Path.ChangeExtension(FileName, DebugInfoExtension));
 			}
 		}
-
-		/// <summary>
-		/// 
-		/// </summary>
-		/// <param name="FileName"></param>
-		public void AddLibraryBuildProduct(string FileName)
-		{
-			string FullFileName = Path.GetFullPath(FileName);
-			if (!LibraryBuildProducts.Contains(FullFileName))
-			{
-				LibraryBuildProducts.Add(FullFileName);
-			}
-		}
 	}
 
 	[Serializable]
@@ -517,6 +499,32 @@ namespace UnrealBuildTool
 				}
 			}
 
+			// If we're precompiling, generate a list of all the files that we depend on
+			if (RulesObject.bPrecompile)
+			{
+				DirectoryReference DependencyListDir;
+				if(RulesObject.ProjectFile == null)
+				{
+					DependencyListDir = DirectoryReference.Combine(UnrealBuildTool.EngineDirectory, RulesObject.Name, RulesObject.Configuration.ToString(), RulesObject.Platform.ToString());
+				}
+				else
+				{
+					DependencyListDir = DirectoryReference.Combine(RulesObject.ProjectFile.Directory, RulesObject.Name, RulesObject.Configuration.ToString(), RulesObject.Platform.ToString());
+				}
+
+				FileReference DependencyListFile;
+				if(RulesObject.bBuildAllModules)
+				{
+					DependencyListFile = FileReference.Combine(DependencyListDir, "DependencyList-AllModules.txt");
+				}
+				else
+				{
+					DependencyListFile = FileReference.Combine(DependencyListDir, "DependencyList.txt");
+				}
+
+				RulesObject.DependencyListFileNames.Add(DependencyListFile);
+			}
+
 			// If we're compiling just a single file, we need to prevent unity builds from running
 			if(bCompilingSingleFile)
 			{
@@ -846,18 +854,6 @@ namespace UnrealBuildTool
 		[NonSerialized]
 		public List<UEBuildBinary> Binaries = new List<UEBuildBinary>();
 
-		/// <summary>
-		/// Binaries which are precompiled for modular targets, but are not build products that the target specifically depends on
-		/// </summary>
-		[NonSerialized]
-		public List<UEBuildBinary> PrecompileOnlyBinaries = new List<UEBuildBinary>();
-
-		/// <summary>
-		/// Modules which are precompiled for monolithic targets, but are not build products that the target specifically depends on
-		/// </summary>
-		[NonSerialized]
-		public List<UEBuildModuleCPP> PrecompileOnlyModules = new List<UEBuildModuleCPP>();
-
 		/// <summary>
 		/// If building only a specific set of modules, these are the modules to build
 		/// </summary>
@@ -1318,15 +1314,16 @@ namespace UnrealBuildTool
 		}
 
 		/// <summary>
-		/// Create a list of all the externally referenced files
+		/// Writes a list of all the externally referenced files required to use the precompiled data for this target
 		/// </summary>
-		/// <param name="Modules">All the modules to include files for</param>
-		/// <param name="Files">Set of referenced files</param>
-		void GetExternalFileList(HashSet<UEBuildModule> Modules, HashSet<FileReference> Files)
+		/// <param name="Location">Path to the dependency list</param>
+		void WriteDependencyList(FileReference Location)
 		{
+			HashSet<FileReference> Files = new HashSet<FileReference>();
+
 			// Get the platform we're building for
 			UEBuildPlatform BuildPlatform = UEBuildPlatform.GetBuildPlatform(Platform);
-			foreach (UEBuildModule Module in Modules)
+			foreach (UEBuildModule Module in Binaries.SelectMany(x => x.Modules))
 			{
 				// Skip artificial modules
 				if(Module.RulesFile == null)
@@ -1428,13 +1425,24 @@ namespace UnrealBuildTool
 						}
 					}
 				}
+
+				// Add all of the runtime dependencies
+				foreach(RuntimeDependency RuntimeDependency in Module.RuntimeDependencies)
+				{
+					Files.Add(RuntimeDependency.Path);
+				}
 			}
+
+			// Write the file
+			Log.TraceInformation("Writing dependency list to {0}", Location);
+			DirectoryReference.CreateDirectory(Location.Directory);
+			FileReference.WriteAllLines(Location, Files.Where(x => x.IsUnderDirectory(UnrealBuildTool.RootDirectory)).Select(x => x.MakeRelativeTo(UnrealBuildTool.RootDirectory).Replace(Path.DirectorySeparatorChar, '/')).OrderBy(x => x));
 		}
 
 		/// <summary>
 		/// Generates a public manifest file for writing out
 		/// </summary>
-		public void GenerateManifest(List<KeyValuePair<FileReference, BuildProductType>> BuildProducts, List<KeyValuePair<FileReference, BuildProductType>> PrecompileOnlyBuildProducts)
+		public void GenerateManifest(List<KeyValuePair<FileReference, BuildProductType>> BuildProducts)
 		{
 			FileReference ManifestPath;
 			if (UnrealBuildTool.IsEngineInstalled() && ProjectFile != null)
@@ -1445,52 +1453,30 @@ namespace UnrealBuildTool
 			{
 				ManifestPath = FileReference.Combine(UnrealBuildTool.EngineDirectory, "Intermediate", "Build", "Manifest.xml");
 			}
-			GenerateManifest(ManifestPath, BuildProducts, PrecompileOnlyBuildProducts);
+			GenerateManifest(ManifestPath, BuildProducts);
 		}
 
 		/// <summary>
 		/// Generates a public manifest file for writing out
 		/// </summary>
-		public void GenerateManifest(FileReference ManifestPath, List<KeyValuePair<FileReference, BuildProductType>> BuildProducts, List<KeyValuePair<FileReference, BuildProductType>> PrecompileOnlyBuildProducts)
+		public void GenerateManifest(FileReference ManifestPath, List<KeyValuePair<FileReference, BuildProductType>> BuildProducts)
 		{
 			BuildManifest Manifest = new BuildManifest();
 
-			if (!Rules.bDisableLinking)
+			// Add the regular build products
+			foreach (KeyValuePair<FileReference, BuildProductType> BuildProductPair in BuildProducts)
 			{
-				// Add the regular build products
-				foreach (KeyValuePair<FileReference, BuildProductType> BuildProductPair in BuildProducts)
-				{
-					if(BuildProductPair.Value == BuildProductType.ImportLibrary)
-					{
-						Manifest.LibraryBuildProducts.Add(BuildProductPair.Key.FullName);
-					}
-					else
-					{
-						Manifest.BuildProducts.Add(BuildProductPair.Key.FullName);
-					}
-				}
-
-				// Add the library-only build products
-				HashSet<FileReference> LibraryBuildProducts = new HashSet<FileReference>(PrecompileOnlyBuildProducts.Select(x => x.Key));
-				if(Rules.bPrecompile && Rules.LinkType == TargetLinkType.Monolithic)
-				{
-					foreach(UEBuildModuleCPP Module in Binaries.Concat(PrecompileOnlyBinaries).SelectMany(x => x.Modules).OfType<UEBuildModuleCPP>().Concat(PrecompileOnlyModules))
-					{
-						if(Module.Rules.bPrecompile)
-						{
-							FileReference PrecompiledManifestLocation = Module.PrecompiledManifestLocation;
-							LibraryBuildProducts.Add(PrecompiledManifestLocation);
+				Manifest.BuildProducts.Add(BuildProductPair.Key.FullName);
+			}
 
-							PrecompiledManifest ModuleManifest = PrecompiledManifest.Read(PrecompiledManifestLocation);
-							foreach(FileReference OutputFile in ModuleManifest.OutputFiles)
-							{
-								LibraryBuildProducts.Add(OutputFile);
-							}
-						}
-					}
-				}
-				Manifest.LibraryBuildProducts.AddRange(LibraryBuildProducts.Select(x => x.FullName));
+			// Add all the dependency lists
+			foreach(FileReference DependencyListFileName in Rules.DependencyListFileNames)
+			{
+				Manifest.BuildProducts.Add(DependencyListFileName.FullName);
+			}
 
+			if (!Rules.bDisableLinking)
+			{
 				// Also add the version file if it's been specified
 				if (VersionFile != null)
 				{
@@ -1531,7 +1517,7 @@ namespace UnrealBuildTool
 		/// <summary>
 		/// Prepare all the receipts this target (all the .target and .modules files). See the VersionManifest class for an explanation of what these files are.
 		/// </summary>
-		void PrepareReceipts(UEToolChain ToolChain, List<KeyValuePair<FileReference, BuildProductType>> BuildProducts, List<KeyValuePair<FileReference, BuildProductType>> PrecompileOnlyBuildProducts, EHotReload HotReload)
+		void PrepareReceipts(UEToolChain ToolChain, List<KeyValuePair<FileReference, BuildProductType>> BuildProducts, EHotReload HotReload)
 		{
 			// If linking is disabled, don't generate any receipt
 			if(Rules.bDisableLinking)
@@ -1580,7 +1566,7 @@ namespace UnrealBuildTool
 			Receipt = new TargetReceipt(TargetName, Platform, Configuration, Version);
 			foreach (KeyValuePair<FileReference, BuildProductType> BuildProductPair in BuildProducts)
 			{
-				if(BuildProductPair.Value != BuildProductType.ImportLibrary)
+				if(BuildProductPair.Value != BuildProductType.BuildResource)
 				{
 					Receipt.AddBuildProduct(BuildProductPair.Key, BuildProductPair.Value);
 				}
@@ -1618,72 +1604,6 @@ namespace UnrealBuildTool
 				}
 			}
 
-			// Add any dependencies of precompiled modules into the receipt
-			if(bPrecompile)
-			{
-				Receipt.PrecompiledBuildDependencies.UnionWith(BuildProducts.Where(x => x.Value == BuildProductType.ImportLibrary).Select(x => x.Key));
-				Receipt.PrecompiledBuildDependencies.UnionWith(PrecompileOnlyBuildProducts.Select(x => x.Key));
-
-				// Add all the precompiled outputs
-				if(Rules.LinkType == TargetLinkType.Monolithic)
-				{
-					foreach(UEBuildModuleCPP Module in Binaries.Concat(PrecompileOnlyBinaries).SelectMany(x => x.Modules).OfType<UEBuildModuleCPP>().Concat(PrecompileOnlyModules))
-					{
-						if(Module.Rules.bPrecompile)
-						{
-							FileReference PrecompiledManifestLocation = Module.PrecompiledManifestLocation;
-							Receipt.PrecompiledBuildDependencies.Add(PrecompiledManifestLocation);
-
-							PrecompiledManifest Manifest = PrecompiledManifest.Read(PrecompiledManifestLocation);
-							foreach(FileReference OutputFile in Manifest.OutputFiles)
-							{
-								Receipt.PrecompiledBuildDependencies.Add(OutputFile);
-							}
-						}
-					}
-				}
-
-				// Find all the modules we need to add runtime dependencies for. This may include precompile-only modules, as well as third party binaries
-				HashSet<UEBuildModule> ReferencedModules = new HashSet<UEBuildModule>(PrecompileOnlyModules);
-				foreach(UEBuildModule Module in PrecompileOnlyModules)
-				{
-					ReferencedModules.UnionWith(Module.GetDependencies(false, false));
-				}
-				foreach(UEBuildBinary Binary in Binaries)
-				{
-					ReferencedModules.UnionWith(Binary.Modules);
-				}
-				foreach(UEBuildBinary Binary in PrecompileOnlyBinaries)
-				{
-					ReferencedModules.UnionWith(Binary.Modules);
-				}
-
-				// Add the runtime dependencies of precompiled modules that are not directly part of this target
-				foreach(UEBuildModule ReferencedModule in ReferencedModules)
-				{
-					if (UniqueLinkedModules.Add(ReferencedModule))
-					{
-						foreach (RuntimeDependency RuntimeDependency in ReferencedModule.RuntimeDependencies)
-						{
-							Receipt.PrecompiledRuntimeDependencies.Add(RuntimeDependency.Path);
-						}
-					}
-				}
-
-				// Add all the files which are required to use the precompiled modules
-				HashSet<FileReference> ExternalFiles = new HashSet<FileReference>();
-				GetExternalFileList(ReferencedModules, ExternalFiles);
-
-				// Convert them into relative to the target receipt
-				foreach(FileReference ExternalFile in ExternalFiles)
-				{
-					if(ExternalFile.IsUnderDirectory(UnrealBuildTool.EngineDirectory) || ExternalFile.IsUnderDirectory(ProjectDirectory))
-					{
-						Receipt.PrecompiledBuildDependencies.Add(ExternalFile);
-					}
-				}
-			}
-
 			// Also add the version file if it's been specified
 			if (VersionFile != null)
 			{
@@ -1695,8 +1615,7 @@ namespace UnrealBuildTool
 			if (!bCompileMonolithic)
 			{
 				// Create the receipts for each folder
-				IEnumerable<UEBuildBinary> AllBinaries = Enumerable.Concat(Binaries, PrecompileOnlyBinaries);
-				foreach (UEBuildBinary Binary in AllBinaries)
+				foreach (UEBuildBinary Binary in Binaries)
 				{
 					if(Binary.Type == UEBuildBinaryType.DynamicLinkLibrary)
 					{
@@ -2049,25 +1968,28 @@ namespace UnrealBuildTool
 
 			if(!ProjectFileGenerator.bGenerateProjectFiles)
 			{
-				// Check the distribution level of all binaries based on the dependencies they have
-				if(ProjectFile == null && !Rules.bOutputPubliclyDistributable)
+				if(!Rules.bDisableLinking)
 				{
-					Dictionary<UEBuildModule, Dictionary<RestrictedFolder, DirectoryReference>> ModuleRestrictedFolderCache = new Dictionary<UEBuildModule, Dictionary<RestrictedFolder, DirectoryReference>>();
-
-					bool bResult = true;
-					foreach (UEBuildBinary Binary in Binaries)
+					// Check the distribution level of all binaries based on the dependencies they have
+					if(ProjectFile == null && !Rules.bOutputPubliclyDistributable)
 					{
-						bResult &= Binary.CheckRestrictedFolders(DirectoryReference.FromFile(ProjectFile), ModuleRestrictedFolderCache);
-					}
+						Dictionary<UEBuildModule, Dictionary<RestrictedFolder, DirectoryReference>> ModuleRestrictedFolderCache = new Dictionary<UEBuildModule, Dictionary<RestrictedFolder, DirectoryReference>>();
 
-					if(!bResult)
-					{
-						throw new BuildException("Unable to create binaries in less restricted locations than their input files.");
+						bool bResult = true;
+						foreach (UEBuildBinary Binary in Binaries)
+						{
+							bResult &= Binary.CheckRestrictedFolders(DirectoryReference.FromFile(ProjectFile), ModuleRestrictedFolderCache);
+						}
+
+						if(!bResult)
+						{
+							throw new BuildException("Unable to create binaries in less restricted locations than their input files.");
+						}
 					}
-				}
 
-				// Check for linking against modules prohibited by the EULA
-				CheckForEULAViolation();
+					// Check for linking against modules prohibited by the EULA
+					CheckForEULAViolation();
+				}
 
 				// Build a mapping from module to its plugin
 				Dictionary<UEBuildModule, UEBuildPlugin> ModuleToPlugin = new Dictionary<UEBuildModule, UEBuildPlugin>();
@@ -2201,7 +2123,6 @@ namespace UnrealBuildTool
 
 			// On Mac and Linux we have actions that should be executed after all the binaries are created
 			TargetToolChain.SetupBundleDependencies(Binaries, TargetName);
-			TargetToolChain.SetupBundleDependencies(PrecompileOnlyBinaries, TargetName);
 
 			// Write out the deployment context, if necessary
 			if(Rules.bDeployAfterCompile && !Rules.bDisableLinking)
@@ -2213,8 +2134,7 @@ namespace UnrealBuildTool
 
 			if (!ProjectFileGenerator.bGenerateProjectFiles)
 			{
-				HashSet<UEBuildModuleCPP> ModulesToGenerateHeadersFor = GatherDependencyModules(OriginalBinaries.Concat(PrecompileOnlyBinaries).ToList());
-				ModulesToGenerateHeadersFor.UnionWith(PrecompileOnlyModules);
+				HashSet<UEBuildModuleCPP> ModulesToGenerateHeadersFor = GatherDependencyModules(OriginalBinaries.ToList());
 
 				if (OnlyModules.Count > 0)
 				{
@@ -2281,45 +2201,19 @@ namespace UnrealBuildTool
 				{
 					Binary.GatherDataForProjectFiles(Rules, GlobalCompileEnvironment);
 				}
-				foreach (UEBuildBinary Binary in PrecompileOnlyBinaries)
-				{
-					Binary.GatherDataForProjectFiles(Rules, GlobalCompileEnvironment);
-				}
 				return ECompilationResult.Succeeded;
 			}
 
-			if(ForeignPlugin != null)
+			// Build the target's binaries.
+			foreach (UEBuildBinary Binary in Binaries)
 			{
-				// Build all the precompiled plugin binaries
-				List<UEBuildBinary> ForeignPluginBinaries = PrecompileOnlyBinaries.Where(x => x.PrimaryModule.RulesFile.IsUnderDirectory(ForeignPlugin.Directory)).ToList();
-				foreach(UEBuildBinary Binary in ForeignPluginBinaries)
-				{
-					OutputItems.AddRange(Binary.Build(Rules, TargetToolChain, GlobalCompileEnvironment, GlobalLinkEnvironment, SharedPCHs, WorkingSet, ActionGraph));
-				}
-
-				List<UEBuildModuleCPP> ForeignPluginModules = PrecompileOnlyModules.Where(x => x.RulesFile.IsUnderDirectory(ForeignPlugin.Directory)).ToList();
-				foreach (UEBuildModuleCPP Module in ForeignPluginModules)
-				{
-					OutputItems.AddRange(Module.Compile(Rules, TargetToolChain, GlobalCompileEnvironment, SharedPCHs, WorkingSet, ActionGraph));
-				}
+				OutputItems.AddRange(Binary.Build(Rules, TargetToolChain, GlobalCompileEnvironment, GlobalLinkEnvironment, SharedPCHs, WorkingSet, ActionGraph));
 			}
-			else
-			{
-				// Build the target's binaries.
-				foreach (UEBuildBinary Binary in Binaries)
-				{
-					OutputItems.AddRange(Binary.Build(Rules, TargetToolChain, GlobalCompileEnvironment, GlobalLinkEnvironment, SharedPCHs, WorkingSet, ActionGraph));
-				}
 
-				// Build all the precompiled binaries
-				foreach (UEBuildBinary Binary in PrecompileOnlyBinaries)
-				{
-					OutputItems.AddRange(Binary.Build(Rules, TargetToolChain, GlobalCompileEnvironment, GlobalLinkEnvironment, SharedPCHs, WorkingSet, ActionGraph));
-				}
-				foreach (UEBuildModuleCPP Module in PrecompileOnlyModules)
-				{
-					OutputItems.AddRange(Module.Compile(Rules, TargetToolChain, GlobalCompileEnvironment, SharedPCHs, WorkingSet, ActionGraph));
-				}
+			// If we're just precompiling a plugin, only include output items which are under that directory
+			if(ForeignPlugin != null)
+			{
+				OutputItems.RemoveAll(x => x.Location.IsUnderDirectory(ForeignPlugin.Directory));
 			}
 
 			// Allow the toolchain to modify the final output items
@@ -2334,19 +2228,10 @@ namespace UnrealBuildTool
 				BuildProducts.AddRange(BinaryBuildProducts);
 			}
 
-			// Also get the precompiled build products
-			List<KeyValuePair<FileReference, BuildProductType>> PrecompileOnlyBuildProducts = new List<KeyValuePair<FileReference, BuildProductType>>();
-			foreach (UEBuildBinary Binary in PrecompileOnlyBinaries)
-			{
-				Dictionary<FileReference, BuildProductType> BinaryBuildProducts = new Dictionary<FileReference, BuildProductType>();
-				Binary.GetBuildProducts(Rules, TargetToolChain, BinaryBuildProducts, GlobalLinkEnvironment.bCreateDebugInfo);
-				PrecompileOnlyBuildProducts.AddRange(BinaryBuildProducts);
-			}
-
 			// Create a receipt for the target
 			if (!ProjectFileGenerator.bGenerateProjectFiles)
 			{
-				PrepareReceipts(TargetToolChain, BuildProducts, PrecompileOnlyBuildProducts, HotReload);
+				PrepareReceipts(TargetToolChain, BuildProducts, HotReload);
 			}
 
 			// Make sure all the checked headers were valid
@@ -2361,14 +2246,20 @@ namespace UnrealBuildTool
 				return ECompilationResult.Canceled;
 			}
 
+			// Build a list of all the externally files
+			foreach(FileReference DependencyListFileName in Rules.DependencyListFileNames)
+			{
+				WriteDependencyList(DependencyListFileName);
+			}
+
 			// If we're only generating the manifest, return now
 			if (BuildConfiguration.bGenerateManifest)
 			{
-				GenerateManifest(BuildProducts, PrecompileOnlyBuildProducts);
+				GenerateManifest(BuildProducts);
 			}
 			foreach(FileReference ManifestFileName in Rules.ManifestFileNames)
 			{
-				GenerateManifest(ManifestFileName, BuildProducts, PrecompileOnlyBuildProducts);
+				GenerateManifest(ManifestFileName, BuildProducts);
 			}
 
 			// Clean any stale modules which exist in multiple output directories. This can lead to the wrong DLL being loaded on Windows.
@@ -2615,10 +2506,10 @@ namespace UnrealBuildTool
 				}
 			}
 
-			// Add all the precompiled modules to the target. In contrast to "Extra Modules", these modules are not compiled into monolithic targets by default.
-			if(bPrecompile)
+			// Add all the modules to the target if necessary.
+			if(Rules.bBuildAllModules)
 			{
-				AddModulesToPrecompile();
+				AddAllValidModulesToTarget();
 			}
 
 			// Add the external and non-C++ referenced modules to the binaries that reference them.
@@ -2642,10 +2533,6 @@ namespace UnrealBuildTool
 					{
 						Binary.SetCreateImportLibrarySeparately(true);
 					}
-					foreach (UEBuildBinary Binary in PrecompileOnlyBinaries)
-					{
-						Binary.SetCreateImportLibrarySeparately(true);
-					}
 				}
 				else
 				{
@@ -2668,7 +2555,6 @@ namespace UnrealBuildTool
 			if (Platform == UnrealTargetPlatform.Mac && !Rules.bIsBuildingConsoleApplication)
 			{
 				TargetToolChain.FixBundleBinariesPaths(this, Binaries);
-				TargetToolChain.FixBundleBinariesPaths(this, PrecompileOnlyBinaries);
 			}
 		}
 
@@ -2942,7 +2828,7 @@ namespace UnrealBuildTool
 		/// <summary>
 		/// Adds all the precompiled modules into the target. Precompiled modules are compiled alongside the target, but not linked into it unless directly referenced.
 		/// </summary>
-		protected void AddModulesToPrecompile()
+		protected void AddAllValidModulesToTarget()
 		{
 			// Find all the modules that are part of the target
 			HashSet<string> PrecompiledModuleNames = new HashSet<string>();
@@ -3106,15 +2992,7 @@ namespace UnrealBuildTool
 			{
 				if(Module.Binary == null)
 				{
-					if(bCompileMonolithic)
-					{
-						PrecompileOnlyModules.Add(Module);
-					}
-					else
-					{
-						Module.Binary = CreateDynamicLibraryForModule(Module);
-						PrecompileOnlyBinaries.Add(Module.Binary);
-					}
+					AddModuleToBinary(Module);
 				}
 			}
 		}
diff --git a/Engine/Source/Programs/UnrealBuildTool/Platform/IOS/IOSToolChain.cs b/Engine/Source/Programs/UnrealBuildTool/Platform/IOS/IOSToolChain.cs
--- a/Engine/Source/Programs/UnrealBuildTool/Platform/IOS/IOSToolChain.cs
+++ b/Engine/Source/Programs/UnrealBuildTool/Platform/IOS/IOSToolChain.cs
@@ -1603,6 +1603,10 @@ namespace UnrealBuildTool
 				Log.TraceWarning("Unable to PostBuildSync, Target has no Rules object");
 				return;
 			}
+			if(Target.Rules.bDisableLinking)
+			{
+				return;
+			}
 
 			IOSProjectSettings ProjectSettings = ((IOSPlatform)UEBuildPlatform.GetBuildPlatform(Target.Platform)).ReadProjectSettings(Target.ProjectFile);
 
diff --git a/Engine/Source/Programs/UnrealBuildTool/Platform/Windows/VCToolChain.cs b/Engine/Source/Programs/UnrealBuildTool/Platform/Windows/VCToolChain.cs
--- a/Engine/Source/Programs/UnrealBuildTool/Platform/Windows/VCToolChain.cs
+++ b/Engine/Source/Programs/UnrealBuildTool/Platform/Windows/VCToolChain.cs
@@ -1491,7 +1491,7 @@ namespace UnrealBuildTool
 		{
 			if (Binary.Type == UEBuildBinaryType.DynamicLinkLibrary)
 			{
-				BuildProducts.Add(FileReference.Combine(Binary.IntermediateDirectory, Binary.OutputFilePath.GetFileNameWithoutExtension() + ".lib"), BuildProductType.ImportLibrary);
+				BuildProducts.Add(FileReference.Combine(Binary.IntermediateDirectory, Binary.OutputFilePath.GetFileNameWithoutExtension() + ".lib"), BuildProductType.BuildResource);
 			}
 			if(Binary.Type == UEBuildBinaryType.Executable && Target.bCreateMapFile)
 			{
diff --git a/Engine/Source/Programs/UnrealBuildTool/System/TargetReceipt.cs b/Engine/Source/Programs/UnrealBuildTool/System/TargetReceipt.cs
--- a/Engine/Source/Programs/UnrealBuildTool/System/TargetReceipt.cs
+++ b/Engine/Source/Programs/UnrealBuildTool/System/TargetReceipt.cs
@@ -27,16 +27,6 @@ namespace UnrealBuildTool
 		/// </summary>
 		DynamicLibrary,
 
-		/// <summary>
-		/// A statically linked library. Not required for the executable to run.
-		/// </summary>
-		StaticLibrary,
-
-		/// <summary>
-		/// An import library. Not required for the executable to run.
-		/// </summary>
-		ImportLibrary,
-
 		/// <summary>
 		/// A symbol file. Not required for the executable to run.
 		/// </summary>
@@ -295,16 +285,6 @@ namespace UnrealBuildTool
 		/// </summary>
 		public RuntimeDependencyList RuntimeDependencies = new RuntimeDependencyList();
 
-		/// <summary>
-		/// All the files which are required to use precompiled binaries with this target
-		/// </summary>
-		public HashSet<FileReference> PrecompiledBuildDependencies = new HashSet<FileReference>();
-
-		/// <summary>
-		/// All the files which are required runtime dependencies for precompiled binaries that are part of this target
-		/// </summary>
-		public HashSet<FileReference> PrecompiledRuntimeDependencies = new HashSet<FileReference>();
-
 		/// <summary>
 		/// Additional build properties passed through from the module rules
 		/// </summary>
@@ -347,8 +327,6 @@ namespace UnrealBuildTool
 				RuntimeDependencies.Add(new RuntimeDependency(OtherRuntimeDependency));
 			}
 			AdditionalProperties.AddRange(Other.AdditionalProperties);
-			PrecompiledBuildDependencies.UnionWith(Other.PrecompiledBuildDependencies);
-			PrecompiledRuntimeDependencies.UnionWith(Other.PrecompiledRuntimeDependencies);
 		}
 
 		/// <summary>
@@ -381,8 +359,6 @@ namespace UnrealBuildTool
 					RuntimeDependencies.Add(OtherRuntimeDependency);
 				}
 			}
-			PrecompiledBuildDependencies.UnionWith(Other.PrecompiledBuildDependencies);
-			PrecompiledRuntimeDependencies.UnionWith(Other.PrecompiledRuntimeDependencies);
 		}
 
 		/// <summary>
@@ -565,28 +541,6 @@ namespace UnrealBuildTool
 				}
 			}
 
-			// Read the precompiled dependencies
-			string[] PrecompiledBuildDependencies;
-			if(RawObject.TryGetStringArrayField("PrecompiledBuildDependencies", out PrecompiledBuildDependencies))
-			{
-				foreach(string PrecompiledBuildDependency in PrecompiledBuildDependencies)
-				{
-					FileReference File = ExpandPathVariables(PrecompiledBuildDependency, EngineDir, ProjectDir);
-					Receipt.PrecompiledBuildDependencies.Add(File);
-				}
-			}
-
-			// Read the precompiled dependencies
-			string[] PrecompiledRuntimeDependencies;
-			if(RawObject.TryGetStringArrayField("PrecompiledRuntimeDependencies", out PrecompiledRuntimeDependencies))
-			{
-				foreach(string PrecompiledRuntimeDependency in PrecompiledRuntimeDependencies)
-				{
-					FileReference File = ExpandPathVariables(PrecompiledRuntimeDependency, EngineDir, ProjectDir);
-					Receipt.PrecompiledRuntimeDependencies.Add(File);
-				}
-			}
-
 			return Receipt;
 		}
 
@@ -670,26 +624,6 @@ namespace UnrealBuildTool
 					Writer.WriteArrayEnd();
 				}
 
-				if(PrecompiledBuildDependencies.Count > 0)
-				{
-					Writer.WriteArrayStart("PrecompiledBuildDependencies");
-					foreach(string PrecompiledBuildDependency in PrecompiledBuildDependencies.Select(x => InsertPathVariables(x, EngineDir, ProjectDir)).OrderBy(x => x))
-					{
-						Writer.WriteValue(PrecompiledBuildDependency);
-					}
-					Writer.WriteArrayEnd();
-				}
-
-				if(PrecompiledRuntimeDependencies.Count > 0)
-				{
-					Writer.WriteArrayStart("PrecompiledRuntimeDependencies");
-					foreach(string PrecompiledRuntimeDependency in PrecompiledRuntimeDependencies.Select(x => InsertPathVariables(x, EngineDir, ProjectDir)).OrderBy(x => x))
-					{
-						Writer.WriteValue(PrecompiledRuntimeDependency);
-					}
-					Writer.WriteArrayEnd();
-				}
-
 				Writer.WriteObjectEnd();
 			}
 		}
diff --git a/Engine/Source/UE4Editor.Target.cs b/Engine/Source/UE4Editor.Target.cs
--- a/Engine/Source/UE4Editor.Target.cs
+++ b/Engine/Source/UE4Editor.Target.cs
@@ -9,7 +9,7 @@ public class UE4EditorTarget : TargetRules
 	{
 		Type = TargetType.Editor;
 		BuildEnvironment = TargetBuildEnvironment.Shared;
-		bPrecompile = true;
+		bBuildAllModules = true;
 		ExtraModuleNames.Add("UE4Game");
 	}
 }
