Exclude editor target files from installed builds, since they leak info about DLLs that have been stripped out.

#rnx

diff --git a/Engine/Build/InstalledEngineFilters.xml b/Engine/Build/InstalledEngineFilters.xml
--- a/Engine/Build/InstalledEngineFilters.xml
+++ b/Engine/Build/InstalledEngineFilters.xml
@@ -60,6 +60,10 @@
 	
 	<!-- Define Editor Filters -->
 	<Property Name="CopyEditorFilter">
+		<!-- Exclude the target files, since we're filtering the output binaries -->
+		Engine/Binaries/.../UE4Editor.target
+		Engine/Binaries/.../UE4Editor-Win64-DebugGame.target
+
 		<!-- This assembly is normally embedded into the UBT executable, but it can technically be rebuilt from an installed build -->
 		Engine/Binaries/ThirdParty/VisualStudio/...
 
