Back out revision 50 from //UE4/Dev-Core/Engine/Build/InstalledEngineBuild.xml

diff --git a/Engine/Build/InstalledEngineBuild.xml b/Engine/Build/InstalledEngineBuild.xml
--- a/Engine/Build/InstalledEngineBuild.xml
+++ b/Engine/Build/InstalledEngineBuild.xml
@@ -118,8 +118,9 @@
 		<Node Name="Compile UnrealHeaderTool Win64" Requires="Update Version Files">
 			<Compile Target="UnrealHeaderTool" Platform="Win64" Configuration="Development" Arguments="-precompile"/>
 		</Node>
-		<Node Name="Compile UE4Editor Win64" Requires="Compile UnrealHeaderTool Win64" Produces="#UE4Editor Win64;#UE4Editor Win64 Unstripped;#UE4Editor Win64 Stripped;#UE4Editor Win64 Unsigned;#UE4Editor Win64 Signed;#UE4Editor Win64 Includes">
-			<Compile Target="UE4Editor" Platform="Win64" Configuration="Development" Tag="#UE4Editor Win64" Arguments="-precompile"/>
+		<Node Name="Compile UE4Editor Win64" Requires="Compile UnrealHeaderTool Win64" Produces="#UE4Editor Win64;#UE4Editor Win64 Unstripped;#UE4Editor Win64 Stripped;#UE4Editor Win64 Unsigned;#UE4Editor Win64 Signed">
+			<Compile Target="UE4Editor" Platform="Win64" Configuration="Development" Tag="#UE4Editor Win64" Arguments="-precompile -allmodules"/>
+			<Compile Target="UE4Editor" Platform="Win64" Configuration="DebugGame" Tag="#UE4Editor Win64" Arguments="-precompile -allmodules"/>
 
 			<Do If="$(EmbedSrcSrvInfo)">
 				<!-- Embed source info into the PDB files. Should be done from this machine to ensure that paths are correct. -->
@@ -147,8 +148,6 @@
 				<Sign Files="#Editor_ToSign_Win64" Tag="#UE4Editor Win64 Signed"/>
 			</Do>
 
-			<!-- Tag the generated includes for this target -->
-			<Tag Files="Engine/Intermediate/Build/Win64/UE4Editor/Inc/...;Engine/Plugins/.../Intermediate/Build/Win64/UE4Editor/Inc/..." Except=".../Timestamp" With="#UE4Editor Win64 Includes"/>
 		</Node>
 	</Agent>
 
@@ -156,10 +155,11 @@
 	<!-- Base editor and tools for compile on Mac -->
 	<Agent Name="Editor Mac" Type="CompileMac;Mac">
 		<Node Name="Compile UnrealHeaderTool Mac" Requires="Update Version Files">
-			<Compile Target="UnrealHeaderTool" Platform="Mac" Configuration="Development" Arguments="-precompile" If="'$(HostPlatform)' != 'Win64'"/>
+			<Compile Target="UnrealHeaderTool" Platform="Mac" Configuration="Development" Arguments="-precompile -allmodules" If="'$(HostPlatform)' != 'Win64'"/>
 		</Node>
-		<Node Name="Compile UE4Editor Mac" Requires="Compile UnrealHeaderTool Mac" Produces="#UE4Editor Mac;#UE4Editor Mac Unstripped;#UE4Editor Mac Stripped;#UE4Editor Mac Unsigned;#UE4Editor Mac Signed;#UE4Editor Mac Includes">
-			<Compile Target="UE4Editor" Platform="Mac" Configuration="Development" Tag="#UE4Editor Mac" Arguments="-precompile"/>
+		<Node Name="Compile UE4Editor Mac" Requires="Compile UnrealHeaderTool Mac" Produces="#UE4Editor Mac;#UE4Editor Mac Unstripped;#UE4Editor Mac Stripped;#UE4Editor Mac Unsigned;#UE4Editor Mac Signed">
+			<Compile Target="UE4Editor" Platform="Mac" Configuration="Development" Tag="#UE4Editor Mac" Arguments="-precompile -allmodules"/>
+			<Compile Target="UE4Editor" Platform="Mac" Configuration="DebugGame" Tag="#UE4Editor Mac" Arguments="-precompile -allmodules"/>
 
 			<Tag Files="#UE4Editor Mac" Filter="$(MacStripFilter)" With="#UE4Editor Mac Unstripped"/>
 
@@ -179,19 +179,17 @@
 				<Copy Files="#UE4Editor Mac Unsigned" From="$(RootDir)" To="$(SavedOutput)" Tag="#Editor_ToSign_Mac"/>
 				<Sign Files="#Editor_ToSign_Mac" Tag="#UE4Editor Mac Signed"/>
 			</Do>
-
-			<!-- Tag the generated includes for this target -->
-			<Tag Files="Engine/Intermediate/Build/Mac/UE4Editor/Inc/...;Engine/Plugins/.../Intermediate/Build/Mac/UE4Editor/Inc/..." Except=".../Timestamp" With="#UE4Editor Mac Includes"/>
 		</Node>
 	</Agent>
 
 	<!-- Base editor and tools for compile on Linux -->
 	<Agent Name="Editor Linux" Type="CompileWin64;Win64">
 		<Node Name="Compile UnrealHeaderTool Linux" Requires="Update Version Files">
-			<Compile Target="UnrealHeaderTool" Platform="Linux" Configuration="Development" Arguments="-precompile"/>
+			<Compile Target="UnrealHeaderTool" Platform="Linux" Configuration="Development" Arguments="-precompile -allmodules"/>
 		</Node>
-		<Node Name="Compile UE4Editor Linux" Requires="Compile UnrealHeaderTool $(HostPlatform)" Produces="#UE4Editor Linux;#UE4Editor Linux Unstripped;#UE4Editor Linux Stripped;#UE4Editor Linux Includes">
-			<Compile Target="UE4Editor" Platform="Linux" Configuration="Development" Tag="#UE4Editor Linux" Arguments="-precompile"/>
+		<Node Name="Compile UE4Editor Linux" Requires="Compile UnrealHeaderTool $(HostPlatform)" Produces="#UE4Editor Linux;#UE4Editor Linux Unstripped;#UE4Editor Linux Stripped">
+			<Compile Target="UE4Editor" Platform="Linux" Configuration="Development" Tag="#UE4Editor Linux" Arguments="-precompile -allmodules"/>
+			<Compile Target="UE4Editor" Platform="Linux" Configuration="DebugGame" Tag="#UE4Editor Linux" Arguments="-precompile -allmodules"/>
 
 			<Tag Files="#UE4Editor Linux" Filter="$(LinuxStripFilter)" With="#UE4Editor Linux Unstripped"/>
 
@@ -205,18 +203,16 @@
 					<Strip BaseDir="$(RootDir)" OutputDir="$(SavedOutput)" Platform="Linux" Files="#UE4Editor Linux Unstripped" Tag="#UE4Editor Linux Stripped"/>
 				</Default>
 			</Switch>
-
-			<!-- Tag the generated includes for this target -->
-			<Tag Files="Engine/Intermediate/Build/Linux/B4D820EA/UE4Editor/Inc/...;Engine/Plugins/.../Intermediate/Build/Linux/B4D820EA/UE4Editor/Inc/..." Except=".../Timestamp" With="#UE4Editor Linux Includes"/>
 		</Node>
 	</Agent>
 
 	<!-- Target Platforms normally built from a Win64 machine -->
 	<Agent Name="Target Platforms Win64" Type="CompileWin64;Win64">
 		<!-- Downstream monolithics and tools Win64 -->
-		<Node Name="Compile UE4Game Win64" Requires="Compile UnrealHeaderTool Win64" Produces="#UE4Game Win64;#UE4Game Win64 Unstripped;#UE4Game Win64 Stripped;#UE4Game Win64 Unsigned;#UE4Game Win64 Signed;#UE4Game Win64 Includes">
+		<Node Name="Compile UE4Game Win64" Requires="Compile UnrealHeaderTool Win64" Produces="#UE4Game Win64;#UE4Game Win64 Unstripped;#UE4Game Win64 Stripped;#UE4Game Win64 Unsigned;#UE4Game Win64 Signed">
 		
 			<ForEach Name="Configuration" Values="$(GameConfigurations)">
+				<Compile Target="UE4Game" Platform="Win64" Configuration="$(Configuration)" Tag="#UE4Game Win64" Arguments="-precompile -allmodules -nolink $(VSCompilerArg) $(TargetDebugInfoArg)"/>
 				<Compile Target="UE4Game" Platform="Win64" Configuration="$(Configuration)" Tag="#UE4Game Win64" Arguments="-precompile $(VSCompilerArg) $(TargetDebugInfoArg)"/>
 			</ForEach>
 
@@ -245,15 +241,13 @@
 				<Copy Files="#UE4Game Win64 Unsigned" From="$(RootDir)" To="$(SavedOutput)" Tag="#Game_ToSign_Win64"/>
 				<Sign Files="#Game_ToSign_Win64" Tag="#UE4Game Win64 Signed"/>
 			</Do>
-
-			<!-- Tag the generated includes for this target -->
-			<Tag Files="Engine/Intermediate/Build/Win64/UE4/Inc/...;Engine/Plugins/.../Intermediate/Build/Win64/UE4/Inc/..." With="#UE4Game Win64 Includes"/>
 		</Node>
 
 		<!-- Downstream monolithics and tools Win32 -->
-		<Node Name="Compile UE4Game Win32" Requires="Compile UnrealHeaderTool Win64" Produces="#UE4Game Win32;#UE4Game Win32 Unstripped;#UE4Game Win32 Stripped;#UE4Game Win32 Unsigned;#UE4Game Win32 Signed;#UE4Game Win32 Includes">
+		<Node Name="Compile UE4Game Win32" Requires="Compile UnrealHeaderTool Win64" Produces="#UE4Game Win32;#UE4Game Win32 Unstripped;#UE4Game Win32 Stripped;#UE4Game Win32 Unsigned;#UE4Game Win32 Signed">
 		
 			<ForEach Name="Configuration" Values="$(GameConfigurations)">
+				<Compile Target="UE4Game" Platform="Win32" Configuration="$(Configuration)" Tag="#UE4Game Win32" Arguments="-precompile -allmodules -nolink $(VSCompilerArg) $(TargetDebugInfoArg)"/>
 				<Compile Target="UE4Game" Platform="Win32" Configuration="$(Configuration)" Tag="#UE4Game Win32" Arguments="-precompile $(VSCompilerArg) $(TargetDebugInfoArg)"/>
 			</ForEach>
 
@@ -283,103 +277,94 @@
 				<Copy Files="#UE4Game Win32 Unsigned" From="$(RootDir)" To="$(SavedOutput)" Tag="#Game_ToSign_Win32"/>
 				<Sign Files="#Game_ToSign_Win32" Tag="#UE4Game Win32 Signed"/>
 			</Do>
-
-			<!-- Tag the generated includes for this target -->
-			<Tag Files="Engine/Intermediate/Build/Win32/UE4/Inc/...;Engine/Plugins/.../Intermediate/Build/Win32/UE4/Inc/..." With="#UE4Game Win32 Includes"/>
 		</Node>
 
 		<!-- Downstream monolithics and tools Linux -->
-		<Node Name="Compile UE4Game Linux" Requires="Compile UnrealHeaderTool $(HostPlatform)" Produces="#UE4Game Linux;#UE4Game Linux Unstripped;#UE4Game Linux Stripped;#UE4Game Linux Includes">
+		<Node Name="Compile UE4Game Linux" Requires="Compile UnrealHeaderTool $(HostPlatform)" Produces="#UE4Game Linux;#UE4Game Linux Unstripped;#UE4Game Linux Stripped">
 		
 			<ForEach Name="Configuration" Values="$(GameConfigurations)">
+				<Compile Target="UE4Game" Platform="Linux" Configuration="$(Configuration)" Tag="#UE4Game Linux" Arguments="-precompile -allmodules -nolink $(TargetDebugInfoArg)"/>
 				<Compile Target="UE4Game" Platform="Linux" Configuration="$(Configuration)" Tag="#UE4Game Linux" Arguments="-precompile $(TargetDebugInfoArg)"/>
 			</ForEach>
 			
 			<Tag Files="#UE4Game Linux" Filter="$(LinuxStripFilter)" With="#UE4Game Linux Unstripped"/>
 			<Strip BaseDir="$(RootDir)" OutputDir="$(SavedOutput)" Platform="Linux" Files="#UE4Game Linux Unstripped" Tag="#UE4Game Linux Stripped"/>
 
-			<!-- Tag the generated includes for this target -->
-			<Tag Files="Engine/Intermediate/Build/Linux/B4D820EA/UE4/Inc/...;Engine/Plugins/.../Intermediate/Build/Linux/B4D820EA/UE4/Inc/..." With="#UE4Game Linux Includes"/>
 		</Node>
 
 		<!-- Downstream monolithics and tools Android-->
-		<Node Name="Compile UE4Game Android" Requires="Compile UnrealHeaderTool $(HostPlatform)" Produces="#UE4Game Android;#UE4Game Android Unstripped;#UE4Game Android Stripped;#UE4Game Android Includes">
+		<Node Name="Compile UE4Game Android" Requires="Compile UnrealHeaderTool $(HostPlatform)" Produces="#UE4Game Android;#UE4Game Android Unstripped;#UE4Game Android Stripped">
 		
 			<ForEach Name="Configuration" Values="$(GameConfigurations)">
+				<Compile Target="UE4Game" Platform="Android" Configuration="$(Configuration)" Tag="#UE4Game Android" Arguments="-precompile -allmodules -nolink -architectures=$(AndroidArchitectures) -gpuarchitectures=$(AndroidGPUArchitectures) $(TargetDebugInfoArg)"/>
 				<Compile Target="UE4Game" Platform="Android" Configuration="$(Configuration)" Tag="#UE4Game Android" Arguments="-precompile -architectures=$(AndroidArchitectures) -gpuarchitectures=$(AndroidGPUArchitectures) $(TargetDebugInfoArg)"/>
 			</ForEach>
 		
 			<Tag Files="#UE4Game Android" Filter="$(AndroidStripFilter)" With="#UE4Game Android Unstripped"/>
 			<Strip BaseDir="$(RootDir)" OutputDir="$(SavedOutput)" Platform="Android" Files="#UE4Game Android Unstripped" Tag="#UE4Game Android Stripped"/>
 
-			<!-- Tag the generated includes for this target -->
-			<Tag Files="Engine/Intermediate/Build/Android/UE4/Inc/...;Engine/Plugins/.../Intermediate/Build/Android/UE4/Inc/..." With="#UE4Game Android Includes"/>
 		</Node>
 
 		<!-- Downstream monolithics and tools HTML5-->
-		<Node Name="Compile UE4Game HTML5" Requires="Compile UnrealHeaderTool $(HostPlatform)" Produces="#UE4Game HTML5;#UE4Game HTML5 Includes">
+		<Node Name="Compile UE4Game HTML5" Requires="Compile UnrealHeaderTool $(HostPlatform)" Produces="#UE4Game HTML5">
 		
 			<ForEach Name="Configuration" Values="$(GameConfigurations)">
+				<Compile Target="UE4Game" Platform="HTML5" Configuration="$(Configuration)" Tag="#UE4Game HTML5" Arguments="-precompile -allmodules -nolink $(TargetDebugInfoArg)"/>
 				<Compile Target="UE4Game" Platform="HTML5" Configuration="$(Configuration)" Tag="#UE4Game HTML5" Arguments="-precompile $(TargetDebugInfoArg)"/>
 			</ForEach>
 		
-			<!-- Tag the generated includes for this target -->
-			<Tag Files="Engine/Intermediate/Build/HTML5/UE4/Inc/...;Engine/Plugins/.../Intermediate/Build/HTML5/UE4/Inc/..." With="#UE4Game HTML5 Includes"/>
 		</Node>
 
 		<!-- Downstream monolithics and tools PS4 -->
-		<Node Name="Compile UE4Game PS4" Requires="Compile UnrealHeaderTool Win64" Produces="#UE4Game PS4;#UE4Game PS4 Includes">
+		<Node Name="Compile UE4Game PS4" Requires="Compile UnrealHeaderTool Win64" Produces="#UE4Game PS4">
 		
 			<ForEach Name="Configuration" Values="$(GameConfigurations)">
+				<Compile Target="UE4Game" Platform="PS4" Configuration="$(Configuration)" Tag="#UE4Game PS4" Arguments="-precompile -allmodules -nolink $(TargetDebugInfoArg)"/>
 				<Compile Target="UE4Game" Platform="PS4" Configuration="$(Configuration)" Tag="#UE4Game PS4" Arguments="-precompile $(TargetDebugInfoArg)"/>
 			</ForEach>
 
-			<!-- Tag the generated includes for this target -->
-			<Tag Files="Engine/Intermediate/Build/PS4/UE4/Inc/...;Engine/Plugins/.../Intermediate/Build/PS4/UE4/Inc/..." With="#UE4Game PS4 Includes"/>
 		</Node>
 
 		<!-- Downstream monolithics and tools XboxOne -->
-		<Node Name="Compile UE4Game XboxOne" Requires="Compile UnrealHeaderTool Win64" Produces="#UE4Game XboxOne;#UE4Game XboxOne Includes">
+		<Node Name="Compile UE4Game XboxOne" Requires="Compile UnrealHeaderTool Win64" Produces="#UE4Game XboxOne">
 		
 			<ForEach Name="Configuration" Values="$(GameConfigurations)">
+				<Compile Target="UE4Game" Platform="XboxOne" Configuration="$(Configuration)" Tag="#UE4Game XboxOne" Arguments="-precompile -allmodules -nolink $(TargetDebugInfoArg)"/>
 				<Compile Target="UE4Game" Platform="XboxOne" Configuration="$(Configuration)" Tag="#UE4Game XboxOne" Arguments="-precompile $(TargetDebugInfoArg)"/>
 			</ForEach>
 			
-			<!-- Tag the generated includes for this target -->
-			<Tag Files="Engine/Intermediate/Build/XboxOne/UE4/Inc/...;Engine/Plugins/.../Intermediate/Build/XboxOne/UE4/Inc/..." With="#UE4Game XboxOne Includes"/>
 		</Node>
 
 		<!-- Downstream monolithics Switch -->
-		<Node Name="Compile UE4Game Switch" Requires="Compile UnrealHeaderTool Win64" Produces="#UE4Game Switch;#UE4Game Switch Includes" If="'$(WithSwitch)' == true">
+		<Node Name="Compile UE4Game Switch" Requires="Compile UnrealHeaderTool Win64" Produces="#UE4Game Switch" If="'$(WithSwitch)' == true">
 		
 			<ForEach Name="Configuration" Values="$(GameConfigurations)">
+				<Compile Target="UE4Game" Platform="Switch" Configuration="$(Configuration)" Tag="#UE4Game Switch" Arguments="-precompile -allmodules -nolink $(TargetDebugInfoArg)"/>
 				<Compile Target="UE4Game" Platform="Switch" Configuration="$(Configuration)" Tag="#UE4Game Switch" Arguments="-precompile $(TargetDebugInfoArg)"/>
 			</ForEach>
 			
-			<!-- Tag the generated includes for this target -->
-			<Tag Files="Engine/Intermediate/Build/Switch/UE4/Inc/...;Engine/Plugins/.../Intermediate/Build/Switch/UE4/Inc/..." With="#UE4Game Switch Includes"/>
 		</Node>
 		
-		<Node Name="Compile UE4Game Lumin" Requires="Compile UnrealHeaderTool Win64" Produces="#UE4Game Lumin;#UE4Game Lumin Unstripped;#UE4Game Lumin Stripped;#UE4Game Lumin Includes" If="'$(WithLumin)' == true">
+		<Node Name="Compile UE4Game Lumin" Requires="Compile UnrealHeaderTool Win64" Produces="#UE4Game Lumin;#UE4Game Lumin Unstripped;#UE4Game Lumin Stripped" If="'$(WithLumin)' == true">
 
 			<ForEach Name="Configuration" Values="$(GameConfigurations)">
+				<Compile Target="UE4Game" Platform="Lumin" Configuration="$(Configuration)" Tag="#UE4Game Lumin" Arguments="-precompile -allmodules -nolink -architectures=$(AndroidArchitectures) -gpuarchitectures=$(LuminGPUArchitectures) $(TargetDebugInfoArg)" AllowParallelExecutor="false"/>
 				<Compile Target="UE4Game" Platform="Lumin" Configuration="$(Configuration)" Tag="#UE4Game Lumin" Arguments="-precompile -architectures=$(AndroidArchitectures) -gpuarchitectures=$(LuminGPUArchitectures) $(TargetDebugInfoArg)" AllowParallelExecutor="false"/>
 			</ForEach>
 			
 			<Tag Files="#UE4Game Lumin" Filter="$(LuminStripFilter)" With="#UE4Game Lumin Unstripped"/>
 			<Strip BaseDir="$(RootDir)" OutputDir="$(SavedOutput)" Platform="Lumin" Files="#UE4Game Lumin Unstripped" Tag="#UE4Game Lumin Stripped"/>
 
-			<!-- Tag the generated includes for this target -->
-			<Tag Files="Engine/Intermediate/Build/Lumin/UE4/Inc/...;Engine/Plugins/.../Intermediate/Build/Lumin/UE4/Inc/..." With="#UE4Game Lumin Includes"/>
 		</Node>
 	</Agent>
 
 	<!-- Target Platforms normally built from a Mac machine -->
 	<Agent Name="Target Platforms Mac" Type="CompileMac;Mac">
 		<!-- Downstream monolithics and tools Mac -->
-		<Node Name="Compile UE4Game Mac" Requires="Compile UnrealHeaderTool Mac" Produces="#UE4Game Mac;#UE4Game Mac Unstripped;#UE4Game Mac Stripped;#UE4Game Mac Unsigned;#UE4Game Mac Signed;#UE4Game Mac Includes">
+		<Node Name="Compile UE4Game Mac" Requires="Compile UnrealHeaderTool Mac" Produces="#UE4Game Mac;#UE4Game Mac Unstripped;#UE4Game Mac Stripped;#UE4Game Mac Unsigned;#UE4Game Mac Signed">
 		
 			<ForEach Name="Configuration" Values="$(GameConfigurations)">
+				<Compile Target="UE4Game" Platform="Mac" Configuration="$(Configuration)" Tag="#UE4Game Mac" Arguments="-precompile -allmodules -nolink $(TargetDebugInfoArg)"/>
 				<Compile Target="UE4Game" Platform="Mac" Configuration="$(Configuration)" Tag="#UE4Game Mac" Arguments="-precompile $(TargetDebugInfoArg)"/>
 			</ForEach>
 			
@@ -392,49 +377,45 @@
 				<Sign Files="#Game_ToSign_Mac" Tag="#UE4Game Mac Signed"/>
 			</Do>
 
-			<!-- Tag the generated includes for this target -->
-			<Tag Files="Engine/Intermediate/Build/Mac/UE4/Inc/...;Engine/Plugins/.../Intermediate/Build/Mac/UE4/Inc/..." With="#UE4Game Mac Includes"/>
 		</Node>
 
 		<!-- Downstream monolithics and tools IOS-->
-		<Node Name="Compile UE4Game IOS" Requires="Compile UnrealHeaderTool Mac" Produces="#UE4Game IOS;#UE4Game IOS Unstripped;#UE4Game IOS Stripped;#UE4Game IOS Includes">
+		<Node Name="Compile UE4Game IOS" Requires="Compile UnrealHeaderTool Mac" Produces="#UE4Game IOS;#UE4Game IOS Unstripped;#UE4Game IOS Stripped">
 		
 			<ForEach Name="Configuration" Values="$(GameConfigurations)">
+				<Compile Target="UE4Game" Platform="IOS" Configuration="$(Configuration)" Tag="#UE4Game IOS" Arguments="-precompile -allmodules -nolink $(TargetDebugInfoArg)"/>
 				<Compile Target="UE4Game" Platform="IOS" Configuration="$(Configuration)" Tag="#UE4Game IOS" Arguments="-precompile $(TargetDebugInfoArg)"/>
 			</ForEach>
 
 			<Tag Files="#UE4Game IOS" Filter="$(IOSStripFilter)" With="#UE4Game IOS Unstripped"/>
 			<Strip BaseDir="$(RootDir)" OutputDir="$(SavedOutput)" Platform="IOS" Files="#UE4Game IOS Unstripped" Tag="#UE4Game IOS Stripped"/>
 
-			<!-- Tag the generated includes for this target -->
-			<Tag Files="Engine/Intermediate/Build/IOS/UE4/Inc/...;Engine/Plugins/.../Intermediate/Build/IOS/UE4/Inc/..." With="#UE4Game IOS Includes"/>
 		</Node>
 
 		<!-- Downstream monolithics and tools TVOS-->
-		<Node Name="Compile UE4Game TVOS" Requires="Compile UnrealHeaderTool Mac" Produces="#UE4Game TVOS;#UE4Game TVOS Unstripped;#UE4Game TVOS Stripped;#UE4Game TVOS Includes">
+		<Node Name="Compile UE4Game TVOS" Requires="Compile UnrealHeaderTool Mac" Produces="#UE4Game TVOS;#UE4Game TVOS Unstripped;#UE4Game TVOS Stripped">
 		
 			<ForEach Name="Configuration" Values="$(GameConfigurations)">
+				<Compile Target="UE4Game" Platform="TVOS" Configuration="$(Configuration)" Tag="#UE4Game TVOS" Arguments="-precompile -allmodules -nolink $(TargetDebugInfoArg)"/>
 				<Compile Target="UE4Game" Platform="TVOS" Configuration="$(Configuration)" Tag="#UE4Game TVOS" Arguments="-precompile $(TargetDebugInfoArg)"/>
 			</ForEach>
 
 			<Tag Files="#UE4Game TVOS" Filter="$(TVOSStripFilter)" With="#UE4Game TVOS Unstripped"/>
 			<Strip BaseDir="$(RootDir)" OutputDir="$(SavedOutput)" Platform="TVOS" Files="#UE4Game TVOS Unstripped" Tag="#UE4Game TVOS Stripped"/>
 
-			<!-- Tag the generated includes for this target -->
-			<Tag Files="Engine/Intermediate/Build/TVOS/UE4/Inc/...;Engine/Plugins/.../Intermediate/Build/TVOS/UE4/Inc/..." With="#UE4Game TVOS Includes"/>
 		</Node>
 		
 		<!-- Downstream monolithics and tools Lumin-->
-		<Node Name="Compile UE4Game LuminMac" Requires="Compile UnrealHeaderTool Mac" Produces="#UE4Game LuminMac;#UE4Game LuminMac Unstripped;#UE4Game LuminMac Stripped;#UE4Game LuminMac Includes">
+		<Node Name="Compile UE4Game LuminMac" Requires="Compile UnrealHeaderTool Mac" Produces="#UE4Game LuminMac;#UE4Game LuminMac Unstripped;#UE4Game LuminMac Stripped">
 
 			<ForEach Name="Configuration" Values="$(GameConfigurations)">
-				<Compile Target="UE4Game" Platform="Lumin" Configuration="$(Configuration)" Tag="#UE4Game LuminMac" Arguments="-precompile -architectures=$(AndroidArchitectures) -gpuarchitectures=$(LuminGPUArchitectures) $(TargetDebugInfoArg)" AllowParallelExecutor="false"/>
+				<Compile Target="UE4Game" Platform="Lumin" Configuration="$(Configuration)" Tag="#UE4Game LuminMac" Arguments="-precompile -allmodules -nolink -architectures=$(AndroidArchitectures) -gpuarchitectures=$(LuminGPUArchitectures)" AllowParallelExecutor="false"/>
+				<Compile Target="UE4Game" Platform="Lumin" Configuration="$(Configuration)" Tag="#UE4Game LuminMac" Arguments="-precompile -architectures=$(AndroidArchitectures) -gpuarchitectures=$(LuminGPUArchitectures)" AllowParallelExecutor="false"/>
 			</ForEach>
-			
+
 			<Tag Files="#UE4Game LuminMac" Filter="$(LuminStripFilter)" With="#UE4Game LuminMac Unstripped"/>
 			<Strip BaseDir="$(RootDir)" OutputDir="$(SavedOutput)" Platform="Lumin" Files="#UE4Game LuminMac Unstripped" Tag="#UE4Game LuminMac Stripped"/>
 
-			<Tag Files="Engine/Intermediate/Build/Lumin/UE4/Inc/...;Engine/Plugins/.../Intermediate/Build/Lumin/UE4/Inc/..." With="#UE4Game LuminMac Includes"/>
 		</Node>
 	</Agent>
 
@@ -667,19 +648,19 @@
 
 	<Agent Name="Installed Build Group Win64" Type="Win64">
 		<!-- Build up a list of requirements as some can be optional -->
-		<Property Name="InstalledRequirements" Value="Compile UnrealHeaderTool Win64;#UE4Editor Win64;#UE4Editor Win64 Unstripped;#UE4Editor Win64 Stripped;#UE4Editor Win64 Unsigned;#UE4Editor Win64 Signed;#UE4Editor Win64 Includes;#Build Tools Win64;#Build Tools CS Binaries"/>
+		<Property Name="InstalledRequirements" Value="Compile UnrealHeaderTool Win64;#UE4Editor Win64;#UE4Editor Win64 Unstripped;#UE4Editor Win64 Stripped;#UE4Editor Win64 Unsigned;#UE4Editor Win64 Signed;#Build Tools Win64;#Build Tools CS Binaries"/>
 		<!-- Optional Target Platform requirements -->
-		<Property Name="InstalledRequirements" Value="$(InstalledRequirements);#UE4Game Win64;#UE4Game Win64 Unstripped;#UE4Game Win64 Stripped;#UE4Game Win64 Unsigned;#UE4Game Win64 Signed;#UE4Game Win64 Includes" If="'$(WithWin64)' == true"/>
-		<Property Name="InstalledRequirements" Value="$(InstalledRequirements);#UE4Game Win32;#UE4Game Win32 Unstripped;#UE4Game Win32 Stripped;#UE4Game Win32 Unsigned;#UE4Game Win32 Signed;#UE4Game Win32 Includes" If="'$(WithWin32)' == true"/>
-		<Property Name="InstalledRequirements" Value="$(InstalledRequirements);#UE4Game Android;#UE4Game Android Unstripped;#UE4Game Android Stripped;#UE4Game Android Includes" If="'$(WithAndroid)' == true"/>
-		<Property Name="InstalledRequirements" Value="$(InstalledRequirements);#UE4Game IOS;#UE4Game IOS Unstripped;#UE4Game IOS Stripped;#UE4Game IOS Includes" If="'$(WithIOS)' == true"/>
-		<Property Name="InstalledRequirements" Value="$(InstalledRequirements);#UE4Game TVOS;#UE4Game TVOS Unstripped;#UE4Game TVOS Stripped;#UE4Game TVOS Includes" If="'$(WithTVOS)' == true"/>
-		<Property Name="InstalledRequirements" Value="$(InstalledRequirements);#UE4Game Linux;#UE4Game Linux Unstripped;#UE4Game Linux Stripped;#UE4Game Linux Includes;#Build Tools Linux" If="'$(WithLinux)' == true"/>
-		<Property Name="InstalledRequirements" Value="$(InstalledRequirements);#UE4Game HTML5;#UE4Game HTML5 Includes" If="'$(WithHTML5)' == true"/>
-		<Property Name="InstalledRequirements" Value="$(InstalledRequirements);#UE4Game PS4;#UE4Game PS4 Includes" If="'$(WithPS4)' == true"/>
-		<Property Name="InstalledRequirements" Value="$(InstalledRequirements);#UE4Game XboxOne;#UE4Game XboxOne Includes" If="'$(WithXboxOne)' == true"/>
-		<Property Name="InstalledRequirements" Value="$(InstalledRequirements);#UE4Game Switch;#UE4Game Switch Includes" If="'$(WithSwitch)' == true"/>
-		<Property Name="InstalledRequirements" Value="$(InstalledRequirements);#UE4Game Lumin;#UE4Game Lumin Includes" If="'$(WithLumin)' == true"/>
+		<Property Name="InstalledRequirements" Value="$(InstalledRequirements);#UE4Game Win64;#UE4Game Win64 Unstripped;#UE4Game Win64 Stripped;#UE4Game Win64 Unsigned;#UE4Game Win64 Signed" If="'$(WithWin64)' == true"/>
+		<Property Name="InstalledRequirements" Value="$(InstalledRequirements);#UE4Game Win32;#UE4Game Win32 Unstripped;#UE4Game Win32 Stripped;#UE4Game Win32 Unsigned;#UE4Game Win32 Signed" If="'$(WithWin32)' == true"/>
+		<Property Name="InstalledRequirements" Value="$(InstalledRequirements);#UE4Game Android;#UE4Game Android Unstripped;#UE4Game Android Stripped" If="'$(WithAndroid)' == true"/>
+		<Property Name="InstalledRequirements" Value="$(InstalledRequirements);#UE4Game IOS;#UE4Game IOS Unstripped;#UE4Game IOS Stripped" If="'$(WithIOS)' == true"/>
+		<Property Name="InstalledRequirements" Value="$(InstalledRequirements);#UE4Game TVOS;#UE4Game TVOS Unstripped;#UE4Game TVOS Stripped" If="'$(WithTVOS)' == true"/>
+		<Property Name="InstalledRequirements" Value="$(InstalledRequirements);#UE4Game Linux;#UE4Game Linux Unstripped;#UE4Game Linux Stripped;#Build Tools Linux" If="'$(WithLinux)' == true"/>
+		<Property Name="InstalledRequirements" Value="$(InstalledRequirements);#UE4Game HTML5" If="'$(WithHTML5)' == true"/>
+		<Property Name="InstalledRequirements" Value="$(InstalledRequirements);#UE4Game PS4" If="'$(WithPS4)' == true"/>
+		<Property Name="InstalledRequirements" Value="$(InstalledRequirements);#UE4Game XboxOne" If="'$(WithXboxOne)' == true"/>
+		<Property Name="InstalledRequirements" Value="$(InstalledRequirements);#UE4Game Switch" If="'$(WithSwitch)' == true"/>
+		<Property Name="InstalledRequirements" Value="$(InstalledRequirements);#UE4Game Lumin" If="'$(WithLumin)' == true"/>
 		<Property Name="InstalledRequirements" Value="$(InstalledRequirements);#CompressedDDCWin64" If="'$(WithDDC)' == true"/>
 		<!-- Feature packs -->
 		<Property Name="InstalledRequirements" Value="$(InstalledRequirements);Make Feature Packs"/>
@@ -696,79 +677,83 @@
 
 			<!-- Tag files that always need to be part of an installed build -->
 			<Tag Files="#Compile UnrealHeaderTool Win64" With="#Installed Win64"/>
-			<Tag Files="#UE4Editor Win64;#UE4Editor Win64 Includes" Except="#UE4Editor Win64 Unstripped;#UE4Editor Win64 Unsigned" With="#Installed Win64"/>
+			<Tag Files="#UE4Editor Win64" Except="#UE4Editor Win64 Unstripped;#UE4Editor Win64 Unsigned" With="#Installed Win64"/>
 			<Tag Files="#UE4Editor Win64 Stripped;#UE4Editor Win64 Signed" With="#Saved Output"/>
 			<Tag Files="#Build Tools Win64;#Make Feature Packs" With="#Installed Win64"/>
 
 			<!-- Tag enabled target platforms -->
 			<Do If="'$(WithWin64)' == true">
-				<Tag Files="#UE4Game Win64;#UE4Game Win64 Includes" Except="#UE4Game Win64 Unstripped;#UE4Game Win64 Unsigned" With="#Installed Win64"/>
+				<Tag Files="#UE4Game Win64" Except="#UE4Game Win64 Unstripped;#UE4Game Win64 Unsigned" With="#Installed Win64"/>
 				<Tag Files="#UE4Game Win64 Stripped;#UE4Game Win64 Signed" With="#Saved Output"/>
 				<Property Name="CopyInstalledFilter" Value="$(CopyInstalledFilter);$(CopyWin64Filter)"/>
 				<Property Name="CopyInstalledExceptions" Value="$(CopyInstalledExceptions);$(CopyWin64Exceptions)"/>
 			</Do>
 			<Do If="'$(WithWin32)' == true">
-				<Tag Files="#UE4Game Win32;#UE4Game Win32 Includes" Except="#UE4Game Win32 Unstripped;#UE4Game Win32 Unsigned" With="#Installed Win64"/>
+				<Tag Files="#UE4Game Win32" Except="#UE4Game Win32 Unstripped;#UE4Game Win32 Unsigned" With="#Installed Win64"/>
 				<Tag Files="#UE4Game Win32 Stripped;#UE4Game Win32 Signed" With="#Saved Output"/>
 				<Property Name="CopyInstalledFilter" Value="$(CopyInstalledFilter);$(CopyWin32Filter)"/>
 				<Property Name="CopyInstalledExceptions" Value="$(CopyInstalledExceptions);$(CopyWin32Exceptions)"/>
 			</Do>
 			<Do If="'$(WithAndroid)' == true">
-				<Tag Files="#UE4Game Android;#UE4Game Android Includes" Except="#UE4Game Android Unstripped" With="#Installed Win64"/>
+				<Tag Files="#UE4Game Android" Except="#UE4Game Android Unstripped" With="#Installed Win64"/>
 				<Tag Files="#UE4Game Android Stripped" With="#Saved Output"/>
 				<Property Name="CopyInstalledFilter" Value="$(CopyInstalledFilter);$(CopyAndroidFilter)"/>
 				<Property Name="CopyInstalledExceptions" Value="$(CopyInstalledExceptions);$(CopyAndroidExceptions)"/>
 			</Do>
 			<Do If="'$(WithIOS)' == true">
-				<Tag Files="#UE4Game IOS;#UE4Game IOS Includes" Except="#UE4Game IOS Unstripped" With="#Installed Win64"/>
+				<Tag Files="#UE4Game IOS" Except="#UE4Game IOS Unstripped" With="#Installed Win64"/>
 				<Tag Files="#UE4Game IOS Stripped" With="#Saved Output"/>
 				<Property Name="CopyInstalledFilter" Value="$(CopyInstalledFilter);$(CopyIOSFilter)"/>
 				<Property Name="CopyInstalledExceptions" Value="$(CopyInstalledExceptions);$(CopyIOSExceptions)"/>
 			</Do>
 			<Do If="'$(WithTVOS)' == true">
-				<Tag Files="#UE4Game TVOS;#UE4Game TVOS Includes" Except="#UE4Game TVOS Unstripped" With="#Installed Win64"/>
+				<Tag Files="#UE4Game TVOS" Except="#UE4Game TVOS Unstripped" With="#Installed Win64"/>
 				<Tag Files="#UE4Game TVOS Stripped" With="#Saved Output"/>
 				<Property Name="CopyInstalledFilter" Value="$(CopyInstalledFilter);$(CopyTVOSFilter)"/>
 				<Property Name="CopyInstalledExceptions" Value="$(CopyInstalledExceptions);$(CopyTVOSExceptions)"/>
 			</Do>
 			<Do If="'$(WithLinux)' == true">
-				<Tag Files="#UE4Game Linux;#UE4Game Linux Includes;#Build Tools Linux" Except="#UE4Game Linux Unstripped" With="#Installed Win64"/>
+				<Tag Files="#UE4Game Linux;#Build Tools Linux" Except="#UE4Game Linux Unstripped" With="#Installed Win64"/>
 				<Tag Files="#UE4Game Linux Stripped" With="#Saved Output"/>
 				<Property Name="CopyInstalledFilter" Value="$(CopyInstalledFilter);$(CopyLinuxFilter)"/>
 				<Property Name="CopyInstalledExceptions" Value="$(CopyInstalledExceptions);$(CopyLinuxExceptions)"/>
 			</Do>
 			<Do If="'$(WithHTML5)' == true">
-				<Tag Files="#UE4Game HTML5;#UE4Game HTML5 Includes" With="#Installed Win64"/>
+				<Tag Files="#UE4Game HTML5" With="#Installed Win64"/>
 				<Property Name="CopyInstalledFilter" Value="$(CopyInstalledFilter);$(CopyHTML5Filter)"/>
 				<Property Name="CopyInstalledExceptions" Value="$(CopyInstalledExceptions);$(CopyHTML5Exceptions)"/>
 			</Do>
 			<Do If="'$(WithPS4)' == true">
-				<Tag Files="#UE4Game PS4;#UE4Game PS4 Includes" With="#Installed Win64"/>
+				<Tag Files="#UE4Game PS4" With="#Installed Win64"/>
 				<Property Name="CopyInstalledFilter" Value="$(CopyInstalledFilter);$(CopyPS4Filter)"/>
 				<Property Name="CopyInstalledExceptions" Value="$(CopyInstalledExceptions);$(CopyPS4Exceptions)"/>
 			</Do>
 			<Do If="'$(WithXboxOne)' == true">
-				<Tag Files="#UE4Game XboxOne;#UE4Game XboxOne Includes" With="#Installed Win64"/>
+				<Tag Files="#UE4Game XboxOne" With="#Installed Win64"/>
 				<Property Name="CopyInstalledFilter" Value="$(CopyInstalledFilter);$(CopyXboxOneFilter)"/>
 				<Property Name="CopyInstalledExceptions" Value="$(CopyInstalledExceptions);$(CopyXboxOneExceptions)"/>
 			</Do>
 			<Do If="'$(WithSwitch)' == true">
-				<Tag Files="#UE4Game Switch;#UE4Game Switch Includes" With="#Installed Win64"/>
+				<Tag Files="#UE4Game Switch" With="#Installed Win64"/>
 				<Property Name="CopyInstalledFilter" Value="$(CopyInstalledFilter);$(CopySwitchFilter)"/>
 				<Property Name="CopyInstalledExceptions" Value="$(CopyInstalledExceptions);$(CopySwitchExceptions)"/>
 			</Do>
 			<Do If="'$(WithLumin)' == true">
-				<Tag Files="#UE4Game Lumin;#UE4Game Lumin Includes" With="#Installed Win64"/>
+				<Tag Files="#UE4Game Lumin" With="#Installed Win64"/>
 				<Property Name="CopyInstalledFilter" Value="$(CopyInstalledFilter);$(CopyLuminFilter)"/>
 				<Property Name="CopyInstalledExceptions" Value="$(CopyInstalledExceptions);$(CopyLuminExceptions)"/>
 			</Do>
 
 			<!-- Tag any dependencies from all previous build products -->
 			<Tag Files="#Installed Win64" Filter="*.target" With="#TargetReceipts"/>
-			<TagReceipt Files="#TargetReceipts" RuntimeDependencies="true" PrecompiledBuildDependencies="true" PrecompiledRuntimeDependencies="true" With="#Installed Win64"/>
+			<TagReceipt Files="#TargetReceipts" RuntimeDependencies="true" With="#Installed Win64"/>
+
+			<!-- Tag any additional dependencies -->
+			<Tag Files="#Installed Win64" Filter=".../DependencyList.txt;.../DependencyList-AllModules.txt" With="#DependencyLists"/>
+			<Tag FileLists="#DependencyLists" With="#Installed Win64" />
 
 			<!-- Confidential Exceptions - done last to ensure they never get overwritten -->
-			<Property Name="CopyInstalledExceptions" Value="$(CopyInstalledExceptions);$(ConfidentialExceptions)"/>
+			<Property Name="CopyInstalledExceptions" Value="$(CopyInstalledExceptions);#DependencyLists;$(ConfidentialExceptions)"/>
 			
 			<!-- Tag additional files needed for installed build -->
 			<Tag Files="$(CopyInstalledFilter)" With="#Installed Win64"/>
@@ -818,15 +803,15 @@
 
 	<Agent Name="Installed Build Group Mac" Type="Mac">
 		<!-- Build up a list of requirements as some can be optional -->
-		<Property Name="FilterRequirements" Value="#Compile UnrealHeaderTool Mac;#UE4Editor Mac;#UE4Editor Mac Unstripped;#UE4Editor Mac Stripped;#UE4Editor Mac Unsigned;#UE4Editor Mac Signed;#UE4Editor Mac Includes;#Build Tools Mac;#Build Tools CS Binaries"/>
+		<Property Name="FilterRequirements" Value="#Compile UnrealHeaderTool Mac;#UE4Editor Mac;#UE4Editor Mac Unstripped;#UE4Editor Mac Stripped;#UE4Editor Mac Unsigned;#UE4Editor Mac Signed;#Build Tools Mac;#Build Tools CS Binaries"/>
 
 		<!-- Optional Target Platform requirements -->
-		<Property Name="FilterRequirements" Value="$(FilterRequirements);#UE4Game Mac;#UE4Game Mac Unstripped;#UE4Game Mac Stripped;#UE4Game Mac Unsigned;#UE4Game Mac Signed;#UE4Game Mac Includes" If="'$(WithMac)' == true"/>
-		<Property Name="FilterRequirements" Value="$(FilterRequirements);#UE4Game Android;#UE4Game Android Unstripped;#UE4Game Android Stripped;#UE4Game Android Includes" If="'$(WithAndroid)' == true"/>
-		<Property Name="FilterRequirements" Value="$(FilterRequirements);#UE4Game IOS;#UE4Game IOS Unstripped;#UE4Game IOS Stripped;#UE4Game IOS Includes" If="'$(WithIOS)' == true"/>
-		<Property Name="FilterRequirements" Value="$(FilterRequirements);#UE4Game TVOS;#UE4Game TVOS Unstripped;#UE4Game TVOS Stripped;#UE4Game TVOS Includes" If="'$(WithTVOS)' == true"/>
-		<Property Name="FilterRequirements" Value="$(FilterRequirements);#UE4Game HTML5;#UE4Game HTML5 Includes" If="'$(WithHTML5)' == true"/>
-		<Property Name="FilterRequirements" Value="$(FilterRequirements);#UE4Game LuminMac;#UE4Game LuminMac Includes;#UE4Game LuminMac Unstripped;#UE4Game LuminMac Stripped;#UE4Game LuminMac Includes" If="'$(WithLumin)' == true"/> 
+		<Property Name="FilterRequirements" Value="$(FilterRequirements);#UE4Game Mac;#UE4Game Mac Unstripped;#UE4Game Mac Stripped;#UE4Game Mac Unsigned;#UE4Game Mac Signed" If="'$(WithMac)' == true"/>
+		<Property Name="FilterRequirements" Value="$(FilterRequirements);#UE4Game Android;#UE4Game Android Unstripped;#UE4Game Android Stripped" If="'$(WithAndroid)' == true"/>
+		<Property Name="FilterRequirements" Value="$(FilterRequirements);#UE4Game IOS;#UE4Game IOS Unstripped;#UE4Game IOS Stripped" If="'$(WithIOS)' == true"/>
+		<Property Name="FilterRequirements" Value="$(FilterRequirements);#UE4Game TVOS;#UE4Game TVOS Unstripped;#UE4Game TVOS Stripped" If="'$(WithTVOS)' == true"/>
+		<Property Name="FilterRequirements" Value="$(FilterRequirements);#UE4Game HTML5" If="'$(WithHTML5)' == true"/>
+		<Property Name="FilterRequirements" Value="$(FilterRequirements);#UE4Game LuminMac;#UE4Game LuminMac Unstripped;#UE4Game LuminMac Stripped" If="'$(WithLumin)' == true"/> 
 		<Property Name="FilterRequirements" Value="$(FilterRequirements);#CompressedDDCMac" If="'$(WithDDC)' == true"/>
 
 		<!-- Feature packs -->
@@ -843,42 +828,42 @@
 
 			<!-- Tag files that always need to be part of an installed build -->
 			<Tag Files="#Compile UnrealHeaderTool Mac" With="#Installed Mac"/>
-			<Tag Files="#UE4Editor Mac;#UE4Editor Mac Includes" Except="#UE4Editor Mac Unstripped;#UE4Editor Mac Unsigned" With="#Installed Mac"/>
+			<Tag Files="#UE4Editor Mac" Except="#UE4Editor Mac Unstripped;#UE4Editor Mac Unsigned" With="#Installed Mac"/>
 			<Tag Files="#UE4Editor Mac Stripped;#UE4Editor Mac Signed" With="#Saved Output"/>
 			<Tag Files="#Build Tools Mac;#Make Feature Packs" With="#Installed Mac"/>
 
 			<!-- Tag enabled target platforms -->
 			<Do If="'$(WithMac)' == true">
-				<Tag Files="#UE4Game Mac;#UE4Game Mac Includes" Except="#UE4Game Mac Unstripped;#UE4Game Mac Unsigned" With="#Installed Mac"/>
+				<Tag Files="#UE4Game Mac" Except="#UE4Game Mac Unstripped;#UE4Game Mac Unsigned" With="#Installed Mac"/>
 				<Tag Files="#UE4Game Mac Stripped;#UE4Game Mac Signed" With="#Saved Output"/>
 				<Property Name="CopyInstalledFilter" Value="$(CopyInstalledFilter);$(CopyMacFilter)"/>
 				<Property Name="CopyInstalledExceptions" Value="$(CopyInstalledExceptions);$(CopyMacExceptions)"/>
 			</Do>
 			<Do  If="'$(WithAndroid)' == true">
-				<Tag Files="#UE4Game Android;#UE4Game Android Includes" Except="#UE4Game Android Unstripped" With="#Installed Mac"/>
+				<Tag Files="#UE4Game Android" Except="#UE4Game Android Unstripped" With="#Installed Mac"/>
 				<Tag Files="#UE4Game Android Stripped" With="#Saved Output"/>
 				<Property Name="CopyInstalledFilter" Value="$(CopyInstalledFilter);$(CopyAndroidFilter)"/>
 				<Property Name="CopyInstalledExceptions" Value="$(CopyInstalledExceptions);$(CopyAndroidExceptions)"/>
 			</Do>
 			<Do If="'$(WithIOS)' == true">
-				<Tag Files="#UE4Game IOS;#UE4Game IOS Includes" Except="#UE4Game IOS Unstripped" With="#Installed Mac"/>
+				<Tag Files="#UE4Game IOS" Except="#UE4Game IOS Unstripped" With="#Installed Mac"/>
 				<Tag Files="#UE4Game IOS Stripped" With="#Saved Output"/>
 				<Property Name="CopyInstalledFilter" Value="$(CopyInstalledFilter);$(CopyIOSFilter)"/>
 				<Property Name="CopyInstalledExceptions" Value="$(CopyInstalledExceptions);$(CopyIOSExceptions)"/>
 			</Do>
 			<Do If="'$(WithTVOS)' == true">
-				<Tag Files="#UE4Game TVOS;#UE4Game TVOS Includes" Except="#UE4Game TVOS Unstripped" With="#Installed Mac"/>
+				<Tag Files="#UE4Game TVOS" Except="#UE4Game TVOS Unstripped" With="#Installed Mac"/>
 				<Tag Files="#UE4Game TVOS Stripped" With="#Saved Output"/>
 				<Property Name="CopyInstalledFilter" Value="$(CopyInstalledFilter);$(CopyTVOSFilter)"/>
 				<Property Name="CopyInstalledExceptions" Value="$(CopyInstalledExceptions);$(CopyTVOSExceptions)"/>
 			</Do>
 			<Do If="'$(WithHTML5)' == true">
-				<Tag Files="#UE4Game HTML5;#UE4Game HTML5 Includes" With="#Installed Mac"/>
+				<Tag Files="#UE4Game HTML5" With="#Installed Mac"/>
 				<Property Name="CopyInstalledFilter" Value="$(CopyInstalledFilter);$(CopyHTML5Filter)"/>
 				<Property Name="CopyInstalledExceptions" Value="$(CopyInstalledExceptions);$(CopyHTML5Exceptions)"/>
 			</Do>
 			<Do If="'$(WithLumin)' == true">
-				<Tag Files="#UE4Game LuminMac;#UE4Game LuminMac Includes" Except="#UE4Game LuminMac Unstripped" With="#Installed Mac"/>
+				<Tag Files="#UE4Game LuminMac" Except="#UE4Game LuminMac Unstripped" With="#Installed Mac"/>
 				<Tag Files="#UE4Game LuminMac Stripped" With="#Saved Output"/>
 				<Property Name="CopyInstalledFilter" Value="$(CopyInstalledFilter);$(CopyLuminFilter)"/>
 				<Property Name="CopyInstalledExceptions" Value="$(CopyInstalledExceptions);$(CopyLuminExceptions)"/>
@@ -886,10 +871,14 @@
 
 			<!-- Tag any dependencies from all previous build products -->
 			<Tag Files="#Installed Mac" Filter="*.target" With="#TargetReceipts"/>
-			<TagReceipt Files="#TargetReceipts" RuntimeDependencies="true" PrecompiledBuildDependencies="true" PrecompiledRuntimeDependencies="true" With="#Installed Mac"/>
+			<TagReceipt Files="#TargetReceipts" RuntimeDependencies="true" With="#Installed Mac"/>
+
+			<!-- Tag any additional dependencies -->
+			<Tag Files="#Installed Mac" Filter=".../DependencyList.txt" With="#DependencyLists"/>
+			<Tag FileLists="#DependencyLists" With="#Installed Mac" />
 
 			<!-- Confidential Exceptions - done last to ensure they never get overwritten -->
-			<Property Name="CopyInstalledExceptions" Value="$(CopyInstalledExceptions);$(ConfidentialExceptions)"/>
+			<Property Name="CopyInstalledExceptions" Value="$(CopyInstalledExceptions);$(ConfidentialExceptions);#DependencyLists"/>
 			
 			<!-- Tag additional files needed for editor and platforms -->
 			<Tag Files="$(CopyInstalledFilter)" With="#Installed Mac"/>
@@ -935,10 +924,10 @@
 
 	<Agent Name="Installed Build Group Linux" Type="Win64">
 		<!-- Build up a list of requirements as some can be optional -->
-		<Property Name="InstalledRequirements" Value="#Compile UnrealHeaderTool Linux;#UE4Editor Linux;#UE4Editor Linux Unstripped;#UE4Editor Linux Stripped;#UE4Editor Linux Includes;#Build Tools Linux;#Build Tools CS Binaries"/>
+		<Property Name="InstalledRequirements" Value="#Compile UnrealHeaderTool Linux;#UE4Editor Linux;#UE4Editor Linux Unstripped;#UE4Editor Linux Stripped;#Build Tools Linux;#Build Tools CS Binaries"/>
 
 		<!-- Optional Target Platform requirements -->
-		<Property Name="InstalledRequirements" Value="$(InstalledRequirements);#UE4Game Linux;#UE4Game Linux Unstripped;#UE4Game Linux Stripped;#UE4Game Linux Includes" If="'$(WithLinux)' == true"/>
+		<Property Name="InstalledRequirements" Value="$(InstalledRequirements);#UE4Game Linux;#UE4Game Linux Unstripped;#UE4Game Linux Stripped" If="'$(WithLinux)' == true"/>
 		<Property Name="InstalledRequirements" Value="$(InstalledRequirements);#CompressedDDCLinux" If="'$(WithDDC)' == true"/>
 
 		<!-- Feature packs -->
@@ -956,13 +945,13 @@
 
 			<!-- Tag files that always need to be part of an installed build -->
 			<Tag Files="#Compile UnrealHeaderTool Linux" With="#Installed Linux"/>
-			<Tag Files="#UE4Editor Linux;#UE4Editor Linux Includes" Except="#UE4Editor Linux Unstripped" With="#Installed Linux"/>
+			<Tag Files="#UE4Editor Linux" Except="#UE4Editor Linux Unstripped" With="#Installed Linux"/>
 			<Tag Files="#UE4Editor Linux Stripped" With="#Saved Output"/>
 			<Tag Files="#Build Tools Linux;#Make Feature Packs" With="#Installed Linux"/>
 
 			<!-- Tag enabled target platforms -->
 			<Do If="'$(WithLinux)' == true">
-				<Tag Files="#UE4Game Linux;#UE4Game Linux Includes;#Build Tools Linux" Except="#UE4Game Linux Unstripped" With="#Installed Linux"/>
+				<Tag Files="#UE4Game Linux;#Build Tools Linux" Except="#UE4Game Linux Unstripped" With="#Installed Linux"/>
 				<Tag Files="#UE4Game Linux Stripped" With="#Saved Output"/>
 				<Property Name="CopyInstalledFilter" Value="$(CopyInstalledFilter);$(CopyLinuxFilter)"/>
 				<Property Name="CopyInstalledExceptions" Value="$(CopyInstalledExceptions);$(CopyLinuxExceptions)"/>
@@ -970,10 +959,14 @@
 
 			<!-- Tag any dependencies from all previous build products -->
 			<Tag Files="#Installed Linux" Filter="*.target" With="#TargetReceipts"/>
-			<TagReceipt Files="#TargetReceipts" RuntimeDependencies="true" PrecompiledBuildDependencies="true" PrecompiledRuntimeDependencies="true" With="#Installed Linux"/>
+			<TagReceipt Files="#TargetReceipts" RuntimeDependencies="true" With="#Installed Linux"/>
+
+			<!-- Tag any additional dependencies -->
+			<Tag Files="#Installed Linux" Filter=".../DependencyList.txt" With="#DependencyLists"/>
+			<Tag FileLists="#DependencyLists" With="#Installed Linux" />
 
 			<!-- Confidential Exceptions - done last to ensure they never get overwritten -->
-			<Property Name="CopyInstalledExceptions" Value="$(CopyInstalledExceptions);$(ConfidentialExceptions)"/>
+			<Property Name="CopyInstalledExceptions" Value="$(CopyInstalledExceptions);$(ConfidentialExceptions);#DependencyLists"/>
 
 			<!-- Tag additional files needed for installed build -->
 			<Tag Files="$(CopyInstalledFilter)" With="#Installed Linux"/>
