Fix Linux dependencies being copied for the Windows editor; they can be added as requirements for the Linux target platform on Windows instead, so it respects the user's chosen platforms.

#rnx

diff --git a/Engine/Build/InstalledEngineFilters.xml b/Engine/Build/InstalledEngineFilters.xml
--- a/Engine/Build/InstalledEngineFilters.xml
+++ b/Engine/Build/InstalledEngineFilters.xml
@@ -206,10 +206,6 @@
 			<!-- Binaries which are explicitly staged by UAT rather than listing runtime dependencies -->
 			Engine/Binaries/ThirdParty/OpenSSL/...
 
-			<!-- Linux cross build Debug tools -->
-			Engine/Binaries/Linux/dump_syms.exe
-			Engine/Binaries/Linux/BreakpadSymbolEncoder.exe
-
 			<!-- Extra tools -->
 			Engine/Extras/3dsMaxScripts/...
 			Engine/Extras/VisualStudioDebugging/UE4.natvis
@@ -408,6 +404,12 @@
 		Engine/Source/ThirdParty/Linux/HaveLinuxDependencies
 	</Property>
 
+	<Property Name="CopyLinuxFilter" Value="$(CopyLinuxFilter)" If="'$(HostPlatform)' == 'Win64'">
+		<!-- Linux cross build Debug tools -->
+		Engine/Binaries/Linux/dump_syms.exe
+		Engine/Binaries/Linux/BreakpadSymbolEncoder.exe
+	</Property>
+
 	<Property Name="CopyLinuxExceptions">
 	</Property>
 
