#DEV-VR - Copying all relevant files for Resonance Audio

diff --git a/Engine/Build/InstalledEngineFilters.xml b/Engine/Build/InstalledEngineFilters.xml
--- a/Engine/Build/InstalledEngineFilters.xml
+++ b/Engine/Build/InstalledEngineFilters.xml
@@ -456,7 +456,7 @@
 	<Property Name="CopyLuminFilter">
 		<!-- Build Files -->
 		Engine/Build/Lumin/...
-		Engine/Source/ThirdParty/ResonanceAudioApi/Lumin/ResonanceAudio_LPL.xml
+		Engine/Source/ThirdParty/ResonanceAudioApi/...
 	</Property>
 	
 	<Property Name="CopyLuminExceptions">
