Fix warnings building DDC on Windows, due to DeploymentServer.exe not being copied into the correct output directory.

diff --git a/Engine/Build/InstalledEngineBuild.xml b/Engine/Build/InstalledEngineBuild.xml
--- a/Engine/Build/InstalledEngineBuild.xml
+++ b/Engine/Build/InstalledEngineBuild.xml
@@ -539,7 +539,7 @@
 			<Delete Files="$(DDCDir)\..."/>
 			<!-- Build up a list of files needed to build DDC -->
 			<CsCompile Project="Engine/Source/Programs/UnrealBuildTool/UnrealBuildTool.csproj" Configuration="Development" Platform="AnyCPU" Tag="#ToCopy" TagReferences="#ToCopy" EnumerateOnly="true"/>
-			<Tag Files="#Compile UnrealHeaderTool Win64;#UE4Editor Win64;#Build Tools Win64;#Build Tools CS" With="#ToCopy"/>
+			<Tag Files="#Compile UnrealHeaderTool Win64;#UE4Editor Win64;#Build Tools Win64" With="#ToCopy"/>
 			<Tag Files="#ToCopy" Filter="*.target" With="#TargetReceipts"/>
 			<TagReceipt Files="#TargetReceipts" RuntimeDependencies="true" With="#ToCopy"/>
 			<Tag Files="Engine\Content\..." Except="*.psd;*.pdn;*.fbx;*.po" With="#ToCopy"/>
@@ -560,6 +560,7 @@
 			<Tag Files="Engine\Binaries\Win64\NotForLicensees\UE4Editor-DDCUtils.dll" With="#FilteredCopyList" If="Exists('Engine\Build\NotForLicensees\EpicInternal.txt')"/>
 			<!-- Copy everything to a temporary directory -->
 			<Copy From="$(RootDir)" To="$(DDCDir)" Files="#FilteredCopyList"/>
+			<Copy From="$(CsToolsDir)" To="$(DDCDir)" Files="#Build Tools CS Binaries"/>
 
 			<Command Name="BuildDerivedDataCache" Arguments="-TempDir=&quot;$(DDCDir)&quot; -FeaturePacks=$(ProjectsToBuildDDC) -TargetPlatforms=$(DDCPlatformsWin64) -HostPlatform=Win64 -SavedDir=&quot;$(SavedDir)&quot;"/>
 
