Copied from Dev-Enterprise - Add WebRTCProxy to InstalledEngineBuild script

diff --git a/Engine/Build/InstalledEngineBuild.xml b/Engine/Build/InstalledEngineBuild.xml
--- a/Engine/Build/InstalledEngineBuild.xml
+++ b/Engine/Build/InstalledEngineBuild.xml
@@ -429,6 +429,7 @@
 			<Compile Target="BootstrapPackagedGame" Configuration="Shipping" Platform="Win32" Tag="#Build Tools Win64"/>
 			<Compile Target="BootstrapPackagedGame" Configuration="Shipping" Platform="Win64" Tag="#Build Tools Win64"/>
 			<Compile Target="BuildPatchTool" Configuration="Shipping" Platform="Win64" Tag="#Build Tools Win64"/>
+			<Compile Target="WebRTCProxy" Configuration="Development" Platform="Win64" Tag="#Build Tools Win64"/>
 
 			<Tag Files="Engine/Binaries/Win64/XGEControlWorker.exe" With="#Build Tools Win64"/>
 			<Tag Files="$(RootDir)/Engine/Saved/UnrealBuildTool/BuildConfiguration.Schema.xsd" With="#Build Tools Win64"/>
