Added Multi User server to binary builds

diff --git a/Engine/Build/InstalledEngineBuild.xml b/Engine/Build/InstalledEngineBuild.xml
--- a/Engine/Build/InstalledEngineBuild.xml
+++ b/Engine/Build/InstalledEngineBuild.xml
@@ -426,6 +426,7 @@
 			<Compile Target="UnrealFrontend" Configuration="Development" Platform="Win64" Tag="#Build Tools Win64"/>
 			<Compile Target="UnrealLightmass" Configuration="Development" Platform="Win64" Tag="#Build Tools Win64"/>
 			<Compile Target="UnrealPak" Configuration="Development" Platform="Win64" Tag="#Build Tools Win64"/>
+			<Compile Target="UnrealMultiUserServer" Configuration="Development" Platform="Win64" Tag="#Build Tools Win64"/>
 			<Compile Target="BootstrapPackagedGame" Configuration="Shipping" Platform="Win32" Tag="#Build Tools Win64"/>
 			<Compile Target="BootstrapPackagedGame" Configuration="Shipping" Platform="Win64" Tag="#Build Tools Win64"/>
 			<Compile Target="BuildPatchTool" Configuration="Shipping" Platform="Win64" Tag="#Build Tools Win64"/>
@@ -486,6 +487,7 @@
 			<Compile Target="ShaderCompileWorker" Configuration="Development" Platform="Linux" Tag="#Build Tools Linux"/>
 			<Compile Target="UnrealLightmass" Configuration="Development" Platform="Linux" Tag="#Build Tools Linux"/>
 			<Compile Target="UnrealPak" Configuration="Development" Platform="Linux" Tag="#Build Tools Linux"/>
+			<Compile Target="UnrealMultiUserServer" Configuration="Development" Platform="Linux" Tag="#Build Tools Linux"/>
 		</Node>
 	</Agent>
 
@@ -497,6 +499,7 @@
 			<Compile Target="UnrealFrontend" Configuration="Development" Platform="Mac" Tag="#Build Tools Mac"/>
 			<Compile Target="UnrealLightmass" Configuration="Development" Platform="Mac" Tag="#Build Tools Mac"/>
 			<Compile Target="UnrealPak" Configuration="Development" Platform="Mac" Tag="#Build Tools Mac"/>
+			<Compile Target="UnrealMultiUserServer" Configuration="Development" Platform="Mac" Tag="#Build Tools Mac"/>
 			<Compile Target="BuildPatchTool" Configuration="Shipping" Platform="Mac" Tag="#Build Tools Mac"/>
 			<Compile Target="DsymExporter" Configuration="Development" Platform="Mac" Tag="#Build Tools Mac"/>
 			<Compile Target="UE4EditorServices" Configuration="Development" Platform="Mac" Tag="#Build Tools Mac"/>
