Add XcodeBuild.sh to filters for installed build.

diff --git a/Engine/Build/InstalledEngineFilters.xml b/Engine/Build/InstalledEngineFilters.xml
--- a/Engine/Build/InstalledEngineFilters.xml
+++ b/Engine/Build/InstalledEngineFilters.xml
@@ -234,6 +234,7 @@
 		Engine/Build/BatchFiles/Mac/FixMonoFiles.sh
 		Engine/Build/BatchFiles/Mac/RunMono.sh
 		Engine/Build/BatchFiles/Mac/RunXBuild.sh
+		Engine/Build/BatchFiles/Mac/XcodeBuild.sh
 		Engine/Build/BatchFiles/RunUAT.command
 		Engine/Build/BatchFiles/RunUAT.sh
 
