Make sure that all the files needed for the vivox plugin on android make it into the binary build

diff --git a/Engine/Build/InstalledEngineFilters.xml b/Engine/Build/InstalledEngineFilters.xml
--- a/Engine/Build/InstalledEngineFilters.xml
+++ b/Engine/Build/InstalledEngineFilters.xml
@@ -330,6 +330,8 @@
 		Engine/Source/ThirdParty/Oculus/OVRPlugin/...
 		Engine/Source/ThirdParty/ResonanceAudioApi/...
 		Engine/Source/ThirdParty/TangoSDK/...
+		Engine/Source/ThirdParty/Vivox/vivox-sdk/Lib/Android/sdk-release.aar
+		Engine/Source/ThirdParty/Vivox/vivox-sdk/VivoxCoreSDK_UPL.xml
 	</Property>
 	
 	<!-- Win64 Specific -->
