LiveCoding: Add natvis helpers library to installed engine build.

diff --git a/Engine/Build/InstalledEngineFilters.xml b/Engine/Build/InstalledEngineFilters.xml
--- a/Engine/Build/InstalledEngineFilters.xml
+++ b/Engine/Build/InstalledEngineFilters.xml
@@ -298,6 +298,9 @@
 		<!-- Default resource files -->
 		Engine/Build/Windows/Resources/...
 
+		<!-- Natvis helpers for live coding -->
+		Engine/Extras/NatvisHelpers/Win64/NatvisHelpers.lib
+
 		<!-- App local dependencies -->
 		Engine/Binaries/ThirdParty/AppLocalDependencies/...
 	</Property>
@@ -306,6 +309,8 @@
 
 	<!-- Win32 -->
 	<Property Name="CopyWin32Filter">
+		<!-- Natvis helpers for live coding -->
+		Engine/Extras/NatvisHelpers/Win32/NatvisHelpers.lib
 	</Property>
 	<Property Name="CopyWin32Exceptions">
 	</Property>
