Added files to list to copy over for installed builds.

diff --git a/Engine/Build/InstalledEngineFilters.xml b/Engine/Build/InstalledEngineFilters.xml
--- a/Engine/Build/InstalledEngineFilters.xml
+++ b/Engine/Build/InstalledEngineFilters.xml
@@ -303,6 +303,16 @@
 
 		<!-- App local dependencies -->
 		Engine/Binaries/ThirdParty/AppLocalDependencies/...
+		
+		<!-- Oculus Libraries & Build Files -->
+		Engine/Source/ThirdParty/Oculus/LibOVRAudio/LibOVRAudio/lib/win64/*.lib
+		Engine/Source/ThirdParty/Oculus/LibOVRAudio/*.tps
+		Engine/Source/ThirdParty/Oculus/LibOVRAvatar/LibOVRAvatar/lib/win64/*.lib
+		Engine/Source/ThirdParty/Oculus/LibOVRAvatar/*.tps
+		Engine/Source/ThirdParty/Oculus/LibOVRPlatform/LibOVRPlatform/lib/*.lib
+		Engine/Source/ThirdParty/Oculus/LibOVRPlatform/*.tps
+		Engine/Source/ThirdParty/Oculus/OVRPlugin/OVRPlugin/lib/Win64/*.lib
+		Engine/Source/ThirdParty/Oculus/OVRPlugin/*.tps
 	</Property>
 	<Property Name="CopyWin64Exceptions">
 	</Property>
@@ -311,6 +321,16 @@
 	<Property Name="CopyWin32Filter">
 		<!-- Natvis helpers for live coding -->
 		Engine/Extras/NatvisHelpers/Win32/NatvisHelpers.lib
+		
+		<!-- Oculus Libraries & Build Files -->
+		Engine/Source/ThirdParty/Oculus/LibOVRAudio/LibOVRAudio/lib/win32/*.lib
+		Engine/Source/ThirdParty/Oculus/LibOVRAudio/*.tps
+		Engine/Source/ThirdParty/Oculus/LibOVRAvatar/LibOVRAvatar/lib/win32/*.lib
+		Engine/Source/ThirdParty/Oculus/LibOVRAvatar/*.tps
+		Engine/Source/ThirdParty/Oculus/LibOVRPlatform/LibOVRPlatform/lib/*.lib
+		Engine/Source/ThirdParty/Oculus/LibOVRPlatform/*.tps
+		Engine/Source/ThirdParty/Oculus/OVRPlugin/OVRPlugin/lib/Win32/*.lib
+		Engine/Source/ThirdParty/Oculus/OVRPlugin/*.tps
 	</Property>
 	<Property Name="CopyWin32Exceptions">
 	</Property>
@@ -330,9 +350,18 @@
 		Engine/Source/ThirdParty/GoogleARCore/...
 		Engine/Source/ThirdParty/GoogleVR/...
 		Engine/Source/ThirdParty/GoogleInstantPreview/InstantPreview.apk
-		Engine/Source/ThirdParty/Oculus/LibOVRPlatform/LibOVRPlatform_APL.xml
-		Engine/Source/ThirdParty/Oculus/LibOVRPlatform/lib/*.so
-		Engine/Source/ThirdParty/Oculus/OVRPlugin/...
+		Engine/Source/ThirdParty/Oculus/LibOVRAudio/LibOVRAudio/lib/*/*.so
+		Engine/Source/ThirdParty/Oculus/LibOVRAudio/*.tps
+		Engine/Source/ThirdParty/Oculus/LibOVRAudio/*.xml
+		Engine/Source/ThirdParty/Oculus/LibOVRAvatar/LibOVRAvatar/lib/*/*.so
+		Engine/Source/ThirdParty/Oculus/LibOVRAvatar/*.tps
+		Engine/Source/ThirdParty/Oculus/LibOVRAvatar/*.xml
+		Engine/Source/ThirdParty/Oculus/LibOVRPlatform/LibOVRPlatform/lib/*.so
+		Engine/Source/ThirdParty/Oculus/LibOVRPlatform/*.tps
+		Engine/Source/ThirdParty/Oculus/LibOVRPlatform/*.xml
+		Engine/Source/ThirdParty/Oculus/OVRPlugin/OVRPlugin/lib/*/.so
+		Engine/Source/ThirdParty/Oculus/OVRPlugin/*.tps
+		Engine/Source/ThirdParty/Oculus/OVRPlugin/*.xml
 		Engine/Source/ThirdParty/ResonanceAudioApi/...
 		Engine/Source/ThirdParty/TangoSDK/...
 		Engine/Source/ThirdParty/Vivox/vivox-sdk/Lib/Android/sdk-release.aar
