Removing unexpected files from the installed build.
UBT should be able to discover these paths automatically through the UEBuildTarget.WriteDependencyList() method, but it wasn't working for some reason for HoloLens, but we addeed too wide of a filter to compensate.

This is workaround for 4.23 preview 6, but we plan to fix it better for the next preview.

diff --git a/Engine/Build/InstalledEngineFilters.xml b/Engine/Build/InstalledEngineFilters.xml
--- a/Engine/Build/InstalledEngineFilters.xml
+++ b/Engine/Build/InstalledEngineFilters.xml
@@ -536,17 +536,16 @@
 		Engine/Binaries/HoloLens/...
 		Engine/Build/HoloLens/...
 		Engine/Intermediate/Deploy/arm64/...
-		Engine/Source/ThirdParty/zlib/...
-		Engine/Source/ThirdParty/IntelTBB/...
-		Engine/Source/ThirdParty/ICU/...
-		Engine/Source/ThirdParty/FreeType2/...
-		Engine/Source/ThirdParty/libPNG/...
-		Engine/Source/ThirdParty/Ogg/...
-		Engine/Source/ThirdParty/Vorbis/...
-		Engine/Source/ThirdParty/libOpus/...
+		Engine/Source/ThirdParty/zlib/v1.2.8/lib/Win64/VS2015/arm64/...
+		Engine/Source/ThirdParty/IntelTBB/IntelTBB-4.4u3/lib/Win64/vc14/arm64/...
+		Engine/Source/ThirdParty/ICU/icu4c-53_1/HoloLens/...
+		Engine/Source/ThirdParty/FreeType2/FreeType2-2.6/Lib/Win64/VS2015/arm64/...
+		Engine/Source/ThirdParty/libPNG/libPNG-1.5.2/lib/HoloLens/...
+		Engine/Source/ThirdParty/Ogg/libogg-1.2.2/lib/HoloLens/...
+		Engine/Source/ThirdParty/Vorbis/libvorbis-1.3.2/lib/HoloLens/...
+		Engine/Source/ThirdParty/libOpus/opus-1.1/Windows/VS2015/ARM64/...
 		Engine/Source/ThirdParty/WindowsMixedRealityInterop/...
-		Engine/Source/ThirdParty/PhysXVehicles/...
-		Engine/Source/ThirdParty/PhysX3/...
+		Engine/Source/ThirdParty/PhysX3/Lib/HoloLens/...
 	</Property>
 
 	<Property Name="CopyHoloLensExceptions">
