Since we are compiling two architectures for HoloLens is causing the second architecture to overwrite the dependency list for the first file. Save them in different locations.

Originally we had a filter to explicitly copy the necessary files, but the filter was too wide. Now that we have this fix, removing the explicit copy filter.

diff --git a/Engine/Build/InstalledEngineFilters.xml b/Engine/Build/InstalledEngineFilters.xml
--- a/Engine/Build/InstalledEngineFilters.xml
+++ b/Engine/Build/InstalledEngineFilters.xml
@@ -533,19 +533,7 @@
 	<!-- HoloLens -->
 	<Property Name="CopyHoloLensFilter">
 		<!-- Build Files -->
-		Engine/Binaries/HoloLens/...
 		Engine/Build/HoloLens/...
-		Engine/Intermediate/Deploy/arm64/...
-		Engine/Source/ThirdParty/zlib/v1.2.8/lib/Win64/VS2015/arm64/...
-		Engine/Source/ThirdParty/IntelTBB/IntelTBB-4.4u3/lib/Win64/vc14/arm64/...
-		Engine/Source/ThirdParty/ICU/icu4c-53_1/HoloLens/...
-		Engine/Source/ThirdParty/FreeType2/FreeType2-2.6/Lib/Win64/VS2015/arm64/...
-		Engine/Source/ThirdParty/libPNG/libPNG-1.5.2/lib/HoloLens/...
-		Engine/Source/ThirdParty/Ogg/libogg-1.2.2/lib/HoloLens/...
-		Engine/Source/ThirdParty/Vorbis/libvorbis-1.3.2/lib/HoloLens/...
-		Engine/Source/ThirdParty/libOpus/opus-1.1/Windows/VS2015/ARM64/...
-		Engine/Source/ThirdParty/WindowsMixedRealityInterop/...
-		Engine/Source/ThirdParty/PhysX3/Lib/HoloLens/...
 	</Property>
 
 	<Property Name="CopyHoloLensExceptions">
diff --git a/Engine/Source/Programs/UnrealBuildTool/Configuration/UEBuildTarget.cs b/Engine/Source/Programs/UnrealBuildTool/Configuration/UEBuildTarget.cs
--- a/Engine/Source/Programs/UnrealBuildTool/Configuration/UEBuildTarget.cs
+++ b/Engine/Source/Programs/UnrealBuildTool/Configuration/UEBuildTarget.cs
@@ -651,6 +651,11 @@ namespace UnrealBuildTool
 					DependencyListDir = DirectoryReference.Combine(RulesObject.ProjectFile.Directory, "Intermediate", "DependencyLists", RulesObject.Name, RulesObject.Configuration.ToString(), RulesObject.Platform.ToString());
 				}
 
+				if(!String.IsNullOrEmpty(Descriptor.Architecture))
+				{
+					DependencyListDir = new DirectoryReference(DependencyListDir.FullName + Descriptor.Architecture);
+				}
+
 				FileReference DependencyListFile;
 				if(RulesObject.bBuildAllModules)
 				{
