Add platform extension batch files to installed engine build.

diff --git a/Engine/Build/InstalledEngineFilters.xml b/Engine/Build/InstalledEngineFilters.xml
--- a/Engine/Build/InstalledEngineFilters.xml
+++ b/Engine/Build/InstalledEngineFilters.xml
@@ -204,6 +204,7 @@
 			Engine/Build/BatchFiles/Clean.bat
 			Engine/Build/BatchFiles/Rebuild.bat
 			Engine/Build/BatchFiles/RunUAT.bat
+			Engine/Build/BatchFiles/FindPlatformExtensionSources.bat
 			Engine/Build/BatchFiles/GetMSBuildPath.bat
 
 			<!-- Build utilities -->
@@ -248,6 +249,7 @@
 		Engine/Build/BatchFiles/Mac/XcodeBuild.sh
 		Engine/Build/BatchFiles/RunUAT.command
 		Engine/Build/BatchFiles/RunUAT.sh
+		Engine/Build/BatchFiles/FindPlatformExtensionSources.sh
 
 		<!-- Mono -->
 		Engine/Binaries/ThirdParty/Mono/Mac/...
@@ -286,6 +288,7 @@
 			Engine/Build/BatchFiles/Linux/RunMono.sh
 			Engine/Build/BatchFiles/Linux/RunXBuild.sh
 			Engine/Build/BatchFiles/RunUAT.sh
+			Engine/Build/BatchFiles/FindPlatformExtensionSources.sh
 
 			<!-- Debug tools -->
 			Engine/Binaries/Linux/dump_syms
