Change path for C++ API docs in the installed engine build.

#rnx

diff --git a/Engine/Build/InstalledEngineFilters.xml b/Engine/Build/InstalledEngineFilters.xml
--- a/Engine/Build/InstalledEngineFilters.xml
+++ b/Engine/Build/InstalledEngineFilters.xml
@@ -150,7 +150,7 @@
 	<Property Name="CopyEditorFilter" Value="$(CopyEditorFilter);Engine/Source/UE4Client.Target.cs" If="$(WithClient)"/>
 	<Property Name="CopyEditorFilter" Value="$(CopyEditorFilter);Engine/Source/UE4Server.Target.cs" If="$(WithServer)"/>
 	<!-- Optional API docs -->
-	<Property Name="CopyEditorFilter" Value="$(CopyEditorFilter);Engine/Documentation/Builds/CodeAPI-HTML.tgz" If="Exists('Engine/Documentation/Builds/CodeAPI-HTML.tgz')"/>
+	<Property Name="CopyEditorFilter" Value="$(CopyEditorFilter);Engine/Documentation/Builds/CppAPI-HTML.tgz" If="Exists('Engine/Documentation/Builds/CppAPI-HTML.tgz')"/>
 	<Property Name="CopyEditorFilter" Value="$(CopyEditorFilter);Engine/Documentation/Builds/BlueprintAPI-HTML.tgz" If="Exists('Engine/Documentation/Builds/BlueprintAPI-HTML.tgz')"/>
 	<!-- Don't add VR template for Mac and Linux -->
 	<Property Name="CopyEditorFilter" Value="$(CopyEditorFilter);Templates/TP_VirtualRealityBP/..." If="'$(HostPlatform)' != 'Linux'"/>
