Remove HTML5 from installed build script.

#rnx

diff --git a/Engine/Build/InstalledEngineBuild.xml b/Engine/Build/InstalledEngineBuild.xml
--- a/Engine/Build/InstalledEngineBuild.xml
+++ b/Engine/Build/InstalledEngineBuild.xml
@@ -35,7 +35,6 @@
 	<Option Name="WithTVOS" Restrict="true|false" DefaultValue="$(DefaultWithIOS)" Description="Include the tvOS target platform"/>
 	<Option Name="WithLinux" Restrict="true|false" DefaultValue="$(DefaultWithLinux)" Description="Include the Linux target platform"/>
 	<Option Name="WithLinuxAArch64" Restrict="true|false" DefaultValue="false" Description="Include the Linux AArch64 target platform"/>
-	<Option Name="WithHTML5" Restrict="true|false" DefaultValue="$(DefaultWithPlatform)" Description="Include the HTML5 target platform"/>
 	<Option Name="WithPS4" Restrict="true|false" DefaultValue="false" Description="Include the PS4 target platform"/>
 	<Option Name="WithXboxOne" Restrict="true|false" DefaultValue="false" Description="Include the XboxOne target platform"/>
 	<Option Name="WithSwitch" Restrict="true|false" DefaultValue="false" Description="Include the Switch target platform"/>
@@ -325,16 +324,6 @@
 
 		</Node>
 	</Agent>
-	<Agent Name="Target HTML5" Type="CompileWin64;Win64">
-		<Node Name="Compile UE4Game HTML5" Requires="Compile UnrealHeaderTool $(HostPlatform)" Produces="#UE4Game HTML5">
-
-			<ForEach Name="Configuration" Values="$(GameConfigurations)">
-				<Compile Target="UE4Game" Platform="HTML5" Configuration="$(Configuration)" Tag="#UE4Game HTML5" Arguments="-precompile -allmodules -nolink $(TargetDebugInfoArg)"/>
-				<Compile Target="UE4Game" Platform="HTML5" Configuration="$(Configuration)" Tag="#UE4Game HTML5" Arguments="-precompile $(TargetDebugInfoArg)" Clean="false"/>
-			</ForEach>
-
-		</Node>
-	</Agent>
 	<Agent Name="Target PS4" Type="CompileWin64;Win64">
 		<Node Name="Compile UE4Game PS4" Requires="Compile UnrealHeaderTool Win64" Produces="#UE4Game PS4">
 
@@ -529,7 +518,6 @@
 			<CsCompile Project="$(CsToolsDir)/Engine/Source/Programs/IOS/MobileDeviceInterface/MobileDeviceInterface.csproj" Configuration="Development" Platform="AnyCPU" Arguments="/verbosity:minimal /target:Rebuild" Tag="#Build Tools CS Binaries;#Sign Binaries" TagReferences="#Build Tools CS Binaries"/>
 			<CsCompile Project="$(CsToolsDir)/Engine/Source/Programs/IOS/DeploymentInterface/DeploymentInterface.csproj" Configuration="Development" Platform="AnyCPU" Arguments="/verbosity:minimal /target:Rebuild" Tag="#Build Tools CS Binaries;#Sign Binaries" TagReferences="#Build Tools CS Binaries"/>
 			<CsCompile Project="$(CsToolsDir)/Engine/Source/Programs/IOS/DeploymentServerLauncher/DeploymentServerLauncher.csproj" Configuration="Development" Platform="AnyCPU" Arguments="/verbosity:minimal /target:Rebuild" Tag="#Build Tools CS Binaries;#Sign Binaries" TagReferences="#Build Tools CS Binaries"/>
-			<CsCompile Project="$(RootDir)/Engine/Platforms/HTML5/Source/Programs/HTML5/HTML5LaunchHelper/HTML5LaunchHelper.csproj" Configuration="Development" Platform="AnyCPU" Arguments="/verbosity:minimal /target:Rebuild" Tag="#Build Tools CS Binaries;#Sign Binaries" TagReferences="#Build Tools CS Binaries"/>
 
 			<!-- Sign the binaries -->
 			<Sign Files="#Sign Binaries" If="'$(SignExecutables)' == true"/>
@@ -589,7 +577,6 @@
 		<Property Name="DDCPlatformsWin64" Value="$(DDCPlatformsWin64)+TVOS" If="'$(WithTVOS)' == true"/>
 		<Property Name="DDCPlatformsWin64" Value="$(DDCPlatformsWin64)+LinuxNoEditor" If="'$(WithLinux)' == true"/>
 		<Property Name="DDCPlatformsWin64" Value="$(DDCPlatformsWin64)+LinuxAArch64NoEditor" If="'$(WithLinuxAArch64)' == true"/>
-		<Property Name="DDCPlatformsWin64" Value="$(DDCPlatformsWin64)+HTML5" If="'$(WithHTML5)' == true"/>
 		<Property Name="DDCPlatformsWin64" Value="$(DDCPlatformsWin64)+PS4" If="'$(WithPS4)' == true"/>
 		<Property Name="DDCPlatformsWin64" Value="$(DDCPlatformsWin64)+XboxOne" If="'$(WithXboxOne)' == true"/>
 		<Property Name="DDCPlatformsWin64" Value="$(DDCPlatformsWin64)+Switch" If="'$(WithSwitch)' == true"/>
@@ -641,7 +628,6 @@
 		<Property Name="DDCPlatformsMac" Value="$(DDCPlatformsMac)+Android_ATC" If="'$(WithAndroid)' == true"/>
 		<Property Name="DDCPlatformsMac" Value="$(DDCPlatformsMac)+IOS" If="'$(WithIOS)' == true"/>
 		<Property Name="DDCPlatformsMac" Value="$(DDCPlatformsMac)+TVOS" If="'$(WithTVOS)' == true"/>
-		<Property Name="DDCPlatformsMac" Value="$(DDCPlatformsMac)+HTML5" If="'$(WithHTML5)' == true"/>
 		<Property Name="DDCPlatformsMac" Value="$(DDCPlatformsMac)+Lumin" If="'$(WithLumin)' == true"/>
 	</Do>
 
@@ -747,7 +733,6 @@
 		<Property Name="InstalledRequirements" Value="$(InstalledRequirements);#UE4Game TVOS;#UE4Game TVOS Unstripped;#UE4Game TVOS Stripped" If="'$(WithTVOS)' == true"/>
 		<Property Name="InstalledRequirements" Value="$(InstalledRequirements);#UE4Game Linux;#UE4Game Linux Unstripped;#UE4Game Linux Stripped;#Build Tools Linux" If="'$(WithLinux)' == true"/>
 		<Property Name="InstalledRequirements" Value="$(InstalledRequirements);#UE4Game LinuxAArch64;#UE4Game LinuxAArch64 Unstripped;#UE4Game LinuxAArch64 Stripped" If="'$(WithLinuxAArch64)' == true"/>
-		<Property Name="InstalledRequirements" Value="$(InstalledRequirements);#UE4Game HTML5" If="'$(WithHTML5)' == true"/>
 		<Property Name="InstalledRequirements" Value="$(InstalledRequirements);#UE4Game PS4" If="'$(WithPS4)' == true"/>
 		<Property Name="InstalledRequirements" Value="$(InstalledRequirements);#UE4Game XboxOne" If="'$(WithXboxOne)' == true"/>
 		<Property Name="InstalledRequirements" Value="$(InstalledRequirements);#UE4Game Switch" If="'$(WithSwitch)' == true"/>
@@ -816,11 +801,6 @@
 				<Property Name="CopyInstalledFilter" Value="$(CopyInstalledFilter);$(CopyLinuxFilter)"/>
 				<Property Name="CopyInstalledExceptions" Value="$(CopyInstalledExceptions);$(CopyLinuxExceptions)"/>
 			</Do>
-			<Do If="'$(WithHTML5)' == true">
-				<Tag Files="#UE4Game HTML5" With="#Installed Win64"/>
-				<Property Name="CopyInstalledFilter" Value="$(CopyInstalledFilter);$(CopyHTML5Filter)"/>
-				<Property Name="CopyInstalledExceptions" Value="$(CopyInstalledExceptions);$(CopyHTML5Exceptions)"/>
-			</Do>
 			<Do If="'$(WithPS4)' == true">
 				<Tag Files="#UE4Game PS4" With="#Installed Win64"/>
 				<Property Name="CopyInstalledFilter" Value="$(CopyInstalledFilter);$(CopyPS4Filter)"/>
@@ -919,7 +899,6 @@
 		<Property Name="FilterRequirements" Value="$(FilterRequirements);#UE4Game Android;#UE4Game Android Unstripped;#UE4Game Android Stripped" If="'$(WithAndroid)' == true"/>
 		<Property Name="FilterRequirements" Value="$(FilterRequirements);#UE4Game IOS;#UE4Game IOS Unstripped;#UE4Game IOS Stripped" If="'$(WithIOS)' == true"/>
 		<Property Name="FilterRequirements" Value="$(FilterRequirements);#UE4Game TVOS;#UE4Game TVOS Unstripped;#UE4Game TVOS Stripped" If="'$(WithTVOS)' == true"/>
-		<Property Name="FilterRequirements" Value="$(FilterRequirements);#UE4Game HTML5" If="'$(WithHTML5)' == true"/>
 		<Property Name="FilterRequirements" Value="$(FilterRequirements);#UE4Game LuminMac;#UE4Game LuminMac Unstripped;#UE4Game LuminMac Stripped" If="'$(WithLumin)' == true"/> 
 		<Property Name="FilterRequirements" Value="$(FilterRequirements);#CompressedDDCMac" If="'$(WithDDC)' == true"/>
 
@@ -966,11 +945,6 @@
 				<Property Name="CopyInstalledFilter" Value="$(CopyInstalledFilter);$(CopyTVOSFilter)"/>
 				<Property Name="CopyInstalledExceptions" Value="$(CopyInstalledExceptions);$(CopyTVOSExceptions)"/>
 			</Do>
-			<Do If="'$(WithHTML5)' == true">
-				<Tag Files="#UE4Game HTML5" With="#Installed Mac"/>
-				<Property Name="CopyInstalledFilter" Value="$(CopyInstalledFilter);$(CopyHTML5Filter)"/>
-				<Property Name="CopyInstalledExceptions" Value="$(CopyInstalledExceptions);$(CopyHTML5Exceptions)"/>
-			</Do>
 			<Do If="'$(WithLumin)' == true">
 				<Tag Files="#UE4Game LuminMac" Except="#UE4Game LuminMac Unstripped" With="#Installed Mac"/>
 				<Tag Files="#UE4Game LuminMac Stripped" With="#Saved Output"/>
diff --git a/Engine/Build/InstalledEngineFilters.xml b/Engine/Build/InstalledEngineFilters.xml
--- a/Engine/Build/InstalledEngineFilters.xml
+++ b/Engine/Build/InstalledEngineFilters.xml
@@ -471,22 +471,6 @@
 <Property Name="CopyLinuxFilter" Value="$(CopyLinuxFilter);Engine/Extras/ThirdPartyNotUE/putty/..."/>	
 
 <!-- Win64 specific -->
-	<!-- HTML5 -->
-	<Property Name="CopyHTML5Filter">
-		<!-- Build Files -->
-		Engine/Build/HTML5/...
-
-		<!-- SDK -->
-		Engine/Extras/ThirdPartyNotUE/emsdk/emscripten/...
-	</Property>
-	<Property Name="CopyHTML5Filter" Value="$(CopyHTML5Filter)" If="'$(HostPlatform)' == 'Win64'">
-		Engine/Extras/ThirdPartyNotUE/emsdk/Win64/...
-	</Property>
-	<Property Name="CopyHTML5Filter" Value="$(CopyHTML5Filter)" If="'$(HostPlatform)' == 'Mac'">
-		Engine/Extras/ThirdPartyNotUE/emsdk/Mac/...
-	</Property>
-		<Property Name="CopyHTML5Exceptions">
-	</Property>
 
 	<!-- PS4 -->
 	<Property Name="CopyPS4Filter">
