Temporarily disable Lumin on Mac due to missing SDK.

#rnx

diff --git a/Engine/Build/InstalledEngineBuild.xml b/Engine/Build/InstalledEngineBuild.xml
--- a/Engine/Build/InstalledEngineBuild.xml
+++ b/Engine/Build/InstalledEngineBuild.xml
@@ -38,7 +38,7 @@
 	<Option Name="WithPS4" Restrict="true|false" DefaultValue="false" Description="Include the PS4 target platform"/>
 	<Option Name="WithXboxOne" Restrict="true|false" DefaultValue="false" Description="Include the XboxOne target platform"/>
 	<Option Name="WithSwitch" Restrict="true|false" DefaultValue="false" Description="Include the Switch target platform"/>
-	<Option Name="WithLumin" Restrict="true|false" DefaultValue="$(DefaultWithPlatform)" Description="Include the Lumin target platform"/>
+	<Option Name="WithLumin" Restrict="true|false" DefaultValue="false" Description="Include the Lumin target platform"/>
 	<Option Name="WithHoloLens" Restrict="true|false" DefaultValue="$(DefaultWithWindows)" Description="Include the HoloLens target platform"/>
 
 	<!-- Set up target types -->
