Add Templates/TemplateCategories.ini to the installed engine build.

#rnx

diff --git a/Engine/Build/InstalledEngineFilters.xml b/Engine/Build/InstalledEngineFilters.xml
--- a/Engine/Build/InstalledEngineFilters.xml
+++ b/Engine/Build/InstalledEngineFilters.xml
@@ -113,6 +113,7 @@
 		Samples/MobileStarterContent/Content/...
 
 		<!-- Templates -->
+		Templates/TemplateCategories.ini
 		Templates/FP_FirstPerson/...
 		Templates/FP_FirstPersonBP/...
 		Templates/TP_FirstPerson/...
@@ -138,7 +139,7 @@
 		Templates/TP_2DSideScrollerBP/...
 		Templates/TP_VehicleAdv/...
 		Templates/TP_VehicleAdvBP/...
-		
+
 		<!-- Shared template resources -->
 		Templates/TemplateResources/...
 
