Prevent HoloLens and Win32 building when -set:HostPlatformOnly=true is specified on the command line for InstalledEngineBuild.xml.

diff --git a/Engine/Build/InstalledEngineBuild.xml b/Engine/Build/InstalledEngineBuild.xml
--- a/Engine/Build/InstalledEngineBuild.xml
+++ b/Engine/Build/InstalledEngineBuild.xml
@@ -11,8 +11,8 @@
 	<Option Name="AllPlatforms" Restrict="true|false" DefaultValue="false" Description="Include all target platforms by default"/>
 
 	<!-- Get default values for each platform being enabled -->
-	<Property Name="DefaultWithWindows" Value="true"/>
-	<Property Name="DefaultWithWindows" Value="false" If="$(HostPlatformEditorOnly) Or ($(HostPlatformOnly) And '$(HostPlatform)' != 'Win64')"/>
+	<Property Name="DefaultWithWin64" Value="true"/>
+	<Property Name="DefaultWithWin64" Value="false" If="$(HostPlatformEditorOnly) Or ($(HostPlatformOnly) And '$(HostPlatform)' != 'Win64')"/>
 
 	<Property Name="DefaultWithMac" Value="true"/>
 	<Property Name="DefaultWithMac" Value="false" If="$(HostPlatformEditorOnly) Or ($(HostPlatformOnly) And '$(HostPlatform)' != 'Mac')"/>
@@ -25,10 +25,16 @@
 
 	<Property Name="DefaultWithIOS" Value="$(DefaultWithPlatform)"/>
 	<Property Name="DefaultWithIOS" Value="false" If="'$(HostPlatform)' != 'Mac' And !$(AllPlatforms)"/>
-	
+
+	<Property Name="DefaultWithWin32" Value="$(DefaultWithPlatform)"/>
+	<Property Name="DefaultWithWin32" Value="false" If="'$(HostPlatform)' != 'Win64' And !$(AllPlatforms)"/>
+
+	<Property Name="DefaultWithHoloLens" Value="$(DefaultWithPlatform)"/>
+	<Property Name="DefaultWithHoloLens" Value="false" If="'$(HostPlatform)' != 'Win64' And !$(AllPlatforms)"/>
+
 	<!-- Setup default target platforms for Installed build if they haven't been specified on the commandline -->
-	<Option Name="WithWin64" Restrict="true|false" DefaultValue="$(DefaultWithWindows)" Description="Include the Win64 target platform"/>
-	<Option Name="WithWin32" Restrict="true|false" DefaultValue="$(DefaultWithWindows)" Description="Include the Win32 target platform"/>
+	<Option Name="WithWin64" Restrict="true|false" DefaultValue="$(DefaultWithWin64)" Description="Include the Win64 target platform"/>
+	<Option Name="WithWin32" Restrict="true|false" DefaultValue="$(DefaultWithWin32)" Description="Include the Win32 target platform"/>
 	<Option Name="WithMac" Restrict="true|false" DefaultValue="$(DefaultWithMac)" Description="Include the Mac target platform"/>
 	<Option Name="WithAndroid" Restrict="true|false" DefaultValue="$(DefaultWithPlatform)" Description="Include the Android target platform"/>
 	<Option Name="WithIOS" Restrict="true|false" DefaultValue="$(DefaultWithIOS)" Description="Include the iOS target platform"/>
@@ -39,7 +45,7 @@
 	<Option Name="WithXboxOne" Restrict="true|false" DefaultValue="false" Description="Include the XboxOne target platform"/>
 	<Option Name="WithSwitch" Restrict="true|false" DefaultValue="false" Description="Include the Switch target platform"/>
 	<Option Name="WithLumin" Restrict="true|false" DefaultValue="false" Description="Include the Lumin target platform"/>
-	<Option Name="WithHoloLens" Restrict="true|false" DefaultValue="$(DefaultWithWindows)" Description="Include the HoloLens target platform"/>
+	<Option Name="WithHoloLens" Restrict="true|false" DefaultValue="$(DefaultWithHoloLens)" Description="Include the HoloLens target platform"/>
 
 	<!-- Set up target types -->
 	<Option Name="WithClient" Restrict="true|false" DefaultValue="false" Description="Include precompiled client targets"/>
