Add the Templates/Media folder to the installed engine build.

#rnx

diff --git a/Engine/Build/InstalledEngineFilters.xml b/Engine/Build/InstalledEngineFilters.xml
--- a/Engine/Build/InstalledEngineFilters.xml
+++ b/Engine/Build/InstalledEngineFilters.xml
@@ -114,6 +114,7 @@
 
 		<!-- Templates -->
 		Templates/TemplateCategories.ini
+		Templates/Media/...
 		Templates/FP_FirstPerson/...
 		Templates/FP_FirstPersonBP/...
 		Templates/TP_FirstPerson/...
