Undo //UE4/Release-4.24/Engine/Build/... changelist 9597972.  Was a temporary fix for Lumin auto SDK failure that has since been fixed.

Hoping this will fix the launcher issue described in 82231

diff --git a/Engine/Build/InstalledEngineBuild.xml b/Engine/Build/InstalledEngineBuild.xml
--- a/Engine/Build/InstalledEngineBuild.xml
+++ b/Engine/Build/InstalledEngineBuild.xml
@@ -44,7 +44,7 @@
 	<Option Name="WithPS4" Restrict="true|false" DefaultValue="false" Description="Include the PS4 target platform"/>
 	<Option Name="WithXboxOne" Restrict="true|false" DefaultValue="false" Description="Include the XboxOne target platform"/>
 	<Option Name="WithSwitch" Restrict="true|false" DefaultValue="false" Description="Include the Switch target platform"/>
-	<Option Name="WithLumin" Restrict="true|false" DefaultValue="false" Description="Include the Lumin target platform"/>
+	<Option Name="WithLumin" Restrict="true|false" DefaultValue="$(DefaultWithPlatform)" Description="Include the Lumin target platform"/>
 	<Option Name="WithHoloLens" Restrict="true|false" DefaultValue="$(DefaultWithHoloLens)" Description="Include the HoloLens target platform"/>
 
 	<!-- Set up target types -->
