Default LinuxAArch64 platform in binary builds to same default as Linux platform

diff --git a/Engine/Build/InstalledEngineBuild.xml b/Engine/Build/InstalledEngineBuild.xml
--- a/Engine/Build/InstalledEngineBuild.xml
+++ b/Engine/Build/InstalledEngineBuild.xml
@@ -20,6 +20,9 @@
 	<Property Name="DefaultWithLinux" Value="true"/>
 	<Property Name="DefaultWithLinux" Value="false" If="$(HostPlatformEditorOnly) Or ($(HostPlatformOnly) And '$(HostPlatform)' != 'Linux')"/>
 
+	<Property Name="DefaultWithLinuxAArch64" Value="true"/>
+	<Property Name="DefaultWithLinuxAArch64" Value="false" If="$(HostPlatformEditorOnly) Or ($(HostPlatformOnly) And '$(HostPlatform)' != 'Linux')"/>
+
 	<Property Name="DefaultWithPlatform" Value="true"/>
 	<Property Name="DefaultWithPlatform" Value="false" If="$(HostPlatformEditorOnly) Or $(HostPlatformOnly)"/>
 
@@ -40,7 +43,7 @@
 	<Option Name="WithIOS" Restrict="true|false" DefaultValue="$(DefaultWithIOS)" Description="Include the iOS target platform"/>
 	<Option Name="WithTVOS" Restrict="true|false" DefaultValue="$(DefaultWithIOS)" Description="Include the tvOS target platform"/>
 	<Option Name="WithLinux" Restrict="true|false" DefaultValue="$(DefaultWithLinux)" Description="Include the Linux target platform"/>
-	<Option Name="WithLinuxAArch64" Restrict="true|false" DefaultValue="false" Description="Include the Linux AArch64 target platform"/>
+	<Option Name="WithLinuxAArch64" Restrict="true|false" DefaultValue="$(DefaultWithLinuxAArch64)" Description="Include the Linux AArch64 target platform"/>
 	<Option Name="WithPS4" Restrict="true|false" DefaultValue="false" Description="Include the PS4 target platform"/>
 	<Option Name="WithXboxOne" Restrict="true|false" DefaultValue="false" Description="Include the XboxOne target platform"/>
 	<Option Name="WithSwitch" Restrict="true|false" DefaultValue="false" Description="Include the Switch target platform"/>
