Added blank templates to binary builds.

diff --git a/Engine/Build/InstalledEngineFilters.xml b/Engine/Build/InstalledEngineFilters.xml
--- a/Engine/Build/InstalledEngineFilters.xml
+++ b/Engine/Build/InstalledEngineFilters.xml
@@ -115,8 +115,11 @@
 		<!-- Templates -->
 		Templates/TemplateCategories.ini
 		Templates/Media/...
-		Templates/FP_FirstPerson/...
-		Templates/FP_FirstPersonBP/...
+		Templates/TP_Blank/...
+		Templates/TP_BlankBP/...
+		Templates/TP_AEC_BlankBP/...
+		Templates/TP_ME_BlankBP/...
+		Templates/TP_MFG_BlankBP/...
 		Templates/TP_FirstPerson/...
 		Templates/TP_FirstPersonBP/...
 		Templates/TP_Flying/...
