Add a separate option for whether to include Lumin on Mac. Currently defaulted to off.

#rnx

diff --git a/Engine/Build/InstalledEngineBuild.xml b/Engine/Build/InstalledEngineBuild.xml
--- a/Engine/Build/InstalledEngineBuild.xml
+++ b/Engine/Build/InstalledEngineBuild.xml
@@ -47,7 +47,8 @@
 	<Option Name="WithPS4" Restrict="true|false" DefaultValue="false" Description="Include the PS4 target platform"/>
 	<Option Name="WithXboxOne" Restrict="true|false" DefaultValue="false" Description="Include the XboxOne target platform"/>
 	<Option Name="WithSwitch" Restrict="true|false" DefaultValue="false" Description="Include the Switch target platform"/>
-	<Option Name="WithLumin" Restrict="true|false" DefaultValue="$(DefaultWithPlatform)" Description="Include the Lumin target platform"/>
+	<Option Name="WithLumin" Restrict="true|false" DefaultValue="$(DefaultWithPlatform)" Description="Include the Lumin target platform on Windows"/>
+	<Option Name="WithLuminMac" Restrict="true|false" DefaultValue="false" Description="Include the Lumin target platform on Mac"/>
 	<Option Name="WithHoloLens" Restrict="true|false" DefaultValue="$(DefaultWithHoloLens)" Description="Include the HoloLens target platform"/>
 
 	<!-- Set up target types -->
@@ -637,7 +638,7 @@
 		<Property Name="DDCPlatformsMac" Value="$(DDCPlatformsMac)+Android_ATC" If="'$(WithAndroid)' == true"/>
 		<Property Name="DDCPlatformsMac" Value="$(DDCPlatformsMac)+IOS" If="'$(WithIOS)' == true"/>
 		<Property Name="DDCPlatformsMac" Value="$(DDCPlatformsMac)+TVOS" If="'$(WithTVOS)' == true"/>
-		<Property Name="DDCPlatformsMac" Value="$(DDCPlatformsMac)+Lumin" If="'$(WithLumin)' == true"/>
+		<Property Name="DDCPlatformsMac" Value="$(DDCPlatformsMac)+Lumin" If="'$(WithLuminMac)' == true"/>
 	</Do>
 
 	<Agent Name="DDC Group Mac" Type="Mac">
@@ -908,7 +909,7 @@
 		<Property Name="FilterRequirements" Value="$(FilterRequirements);#UE4Game Android;#UE4Game Android Unstripped;#UE4Game Android Stripped" If="'$(WithAndroid)' == true"/>
 		<Property Name="FilterRequirements" Value="$(FilterRequirements);#UE4Game IOS;#UE4Game IOS Unstripped;#UE4Game IOS Stripped" If="'$(WithIOS)' == true"/>
 		<Property Name="FilterRequirements" Value="$(FilterRequirements);#UE4Game TVOS;#UE4Game TVOS Unstripped;#UE4Game TVOS Stripped" If="'$(WithTVOS)' == true"/>
-		<Property Name="FilterRequirements" Value="$(FilterRequirements);#UE4Game LuminMac;#UE4Game LuminMac Unstripped;#UE4Game LuminMac Stripped" If="'$(WithLumin)' == true"/> 
+		<Property Name="FilterRequirements" Value="$(FilterRequirements);#UE4Game LuminMac;#UE4Game LuminMac Unstripped;#UE4Game LuminMac Stripped" If="'$(WithLuminMac)' == true"/>
 		<Property Name="FilterRequirements" Value="$(FilterRequirements);#CompressedDDCMac" If="'$(WithDDC)' == true"/>
 
 		<!-- Feature packs -->
@@ -954,7 +955,7 @@
 				<Property Name="CopyInstalledFilter" Value="$(CopyInstalledFilter);$(CopyTVOSFilter)"/>
 				<Property Name="CopyInstalledExceptions" Value="$(CopyInstalledExceptions);$(CopyTVOSExceptions)"/>
 			</Do>
-			<Do If="'$(WithLumin)' == true">
+			<Do If="'$(WithLuminMac)' == true">
 				<Tag Files="#UE4Game LuminMac" Except="#UE4Game LuminMac Unstripped" With="#Installed Mac"/>
 				<Tag Files="#UE4Game LuminMac Stripped" With="#Saved Output"/>
 				<Property Name="CopyInstalledFilter" Value="$(CopyInstalledFilter);$(CopyLuminFilter)"/>
