HTML5 - more code removal - build scripts section

#rnx

diff --git a/Engine/Build/BatchFiles/Mac/XcodeBuild.sh b/Engine/Build/BatchFiles/Mac/XcodeBuild.sh
--- a/Engine/Build/BatchFiles/Mac/XcodeBuild.sh
+++ b/Engine/Build/BatchFiles/Mac/XcodeBuild.sh
@@ -6,20 +6,19 @@
 # Setup Mono
 source Engine/Build/BatchFiles/Mac/SetupMono.sh Engine/Build/BatchFiles/Mac
 
-# override env if action is specified on command line 
-if [ $1 == "clean" ] 
-	then 
-		ACTION="clean" 
-fi 
+# override env if action is specified on command line
+if [ $1 == "clean" ]; then
+	ACTION="clean"
+fi
 
 
 case $ACTION in
 	"")
 		echo Building $1...
 
-                Platform=""
-                AdditionalFlags=""
-		
+		Platform=""
+		AdditionalFlags=""
+
 		case $CLANG_STATIC_ANALYZER_MODE in
 				"deep")
 					AdditionalFlags+="-SkipActionHistory"
@@ -36,31 +35,28 @@ case $ACTION in
 			;;
 			esac
 
-		case $2 in 
+		case $2 in
 			"iphoneos"|"IOS")
-		        Platform="IOS"
+				Platform="IOS"
 				AdditionalFlags+=" -deploy "
-			;; 
+			;;
 			"appletvos")
-		        Platform="TVOS"
+				Platform="TVOS"
 				AdditionalFlags+=" -deploy "
-			;; 
-  			"iphonesimulator")
-		        	Platform="IOS"
-		         	AdditionalFlags+=" -deploy -simulator"
 			;;
-			"macosx")
-					Platform="Mac"
+			"iphonesimulator")
+				Platform="IOS"
+				AdditionalFlags+=" -deploy -simulator"
 			;;
-			"HTML5")
-				Platform="HTML5"
+			"macosx")
+				Platform="Mac"
 			;;
 			*)
 				Platform="$2"
-				AdditionalFlags+=" -deploy " 
-			;; 
+#				AdditionalFlags+=" -deploy "
+			;;
 		esac
-		
+
 		BuildTasks=$(defaults read com.apple.dt.Xcode IDEBuildOperationMaxNumberOfConcurrentCompileTasks)
 		export NumUBTBuildTasks=$BuildTasks
 
@@ -75,29 +71,26 @@ case $ACTION in
 	"clean")
 		echo "Cleaning $2 $3 $4..."
 
-                Platform=""
-                AdditionalFlags="-clean"
+			Platform=""
+			AdditionalFlags="-clean"
 
-		case $3 in 
+		case $3 in
 			"iphoneos"|"IOS")
-	        	        Platform="IOS"
+				Platform="IOS"
 				AdditionalFlags+=" "
 			;;
 			"iphonesimulator")
-			        Platform="IOS"
-		        	AdditionalFlags+=" -simulator"
+				Platform="IOS"
+				AdditionalFlags+=" -simulator"
 				AdditionalFlags+=" "
-			;; 
+			;;
 			"macosx")
 					Platform="Mac"
 			;;
-			"HTML5")
-				Platform="HTML5"
-			;;
-			*) 
-        			Platform="$3"
+			*)
+				Platform="$3"
 			;;
-		
+
 		esac
 		echo Running command: mono Engine/Binaries/DotNET/UnrealBuildTool.exe $2 $Platform $4 $AdditionalFlags "${@:5}"
 		mono Engine/Binaries/DotNET/UnrealBuildTool.exe $2 $Platform $4 $AdditionalFlags "${@:5}"
@@ -107,6 +100,5 @@ esac
 ExitCode=$?
 if [ $ExitCode -eq 254 ] || [ $ExitCode -eq 255 ] || [ $ExitCode -eq 2 ]; then
 	exit 0
-else
-	exit $ExitCode
 fi
+exit $ExitCode
diff --git a/Engine/Build/PhysXBuild.xml b/Engine/Build/PhysXBuild.xml
--- a/Engine/Build/PhysXBuild.xml
+++ b/Engine/Build/PhysXBuild.xml
@@ -27,7 +27,6 @@
 	<Option Name="WithAndroid-ARM64" DefaultValue="false" Description="Compile for Android ARM64"/>
 	<Option Name="WithAndroid-x86" DefaultValue="false" Description="Compile for Android x86"/>
 	<Option Name="WithAndroid-x64" DefaultValue="false" Description="Compile for Android x64"/>
-	<Option Name="WithHTML5" DefaultValue="false" Description="Compile for HTML5"/>
 	<Option Name="WithSwitch" DefaultValue="false" Description="Compile for Switch"/>
 	<Option Name="ShelvedChange" DefaultValue="" Description="Change to unshelve"/>
 
@@ -105,11 +104,6 @@
 			<Command Name="BuildPhysX" Arguments="-TargetPlatforms=Android-x64 -TargetConfigs=&quot;$(TargetConfigs)&quot; -TargetWindowsCompilers=&quot;$(TargetWindowsCompilers)&quot; $(OtherArgs)"/>
 		</Node>
 	</Agent>
-	<Agent Name="PhysX HTML5" Type="EngineWin64;PhysX_CompileWin64;Win64">
-		<Node Name="Compile PhysX HTML5">
-			<Command Name="BuildPhysX" Arguments="-TargetPlatforms=HTML5 -TargetConfigs=&quot;$(TargetConfigs)&quot; -TargetWindowsCompilers=&quot;$(TargetWindowsCompilers)&quot; $(OtherArgs)"/>
-		</Node>
-	</Agent>
 	<Agent Name="PhysX Mac" Type="EngineMac;PhysX_CompileMac;Mac">
 		<Node Name="Compile PhysX Mac">
 			<Command Name="BuildPhysX" Arguments="-TargetPlatforms=Mac -TargetConfigs=&quot;$(TargetConfigs)&quot; $(OtherArgs)"/>
@@ -145,7 +139,6 @@
 	<Property Name="PhysXTargets" Value="$(PhysXTargets);Compile PhysX Android Arm64" If="$(WithAndroid-ARM64)"/>
 	<Property Name="PhysXTargets" Value="$(PhysXTargets);Compile PhysX Android x86" If="$(WithAndroid-x86)"/>
 	<Property Name="PhysXTargets" Value="$(PhysXTargets);Compile PhysX Android x64" If="$(WithAndroid-x64)"/>
-	<Property Name="PhysXTargets" Value="$(PhysXTargets);Compile PhysX HTML5" If="$(WithHTML5)"/>
 	<Property Name="PhysXTargets" Value="$(PhysXTargets);Compile PhysX Mac" If="$(WithMac)"/>
 	<Property Name="PhysXTargets" Value="$(PhysXTargets);Compile PhysX IOS" If="$(WithIOS)"/>
 	<Property Name="PhysXTargets" Value="$(PhysXTargets);Compile PhysX TVOS" If="$(WithTVOS)"/>
@@ -160,7 +153,6 @@
 	<Notify Nodes="Compile PhysX XboxOne" Users="$(AllEmailTargets);[REDACTED_EMAIL]"/>
 	<Notify Nodes="Compile PhysX Linux;Compile PhysX Linux x86;Compile PhysX Linux ARM32;Compile PhysX Linux ARM64" Users="$(AllEmailTargets);[REDACTED_EMAIL]"/>
 	<Notify Nodes="Compile PhysX Android Armv7;Compile PhysX Android Arm64;Compile PhysX Android x86;Compile PhysX Android x64" Users="$(AllEmailTargets);[REDACTED_EMAIL]"/>
-	<Notify Nodes="Compile PhysX HTML5" Users="$(AllEmailTargets);[REDACTED_EMAIL]"/>
 	<Notify Nodes="Compile PhysX Mac;Compile PhysX IOS;Compile PhysX TVOS" Users="$(AllEmailTargets);[REDACTED_EMAIL]"/>
 	<Notify Nodes="Compile PhysX Switch" Users="$(AllEmailTargets);[REDACTED_EMAIL]"/>
 </BuildGraph>
