[NO NON-PII COMMIT MESSAGE]

diff --git a/Engine/Build/InstalledEngineBuild.xml b/Engine/Build/InstalledEngineBuild.xml
--- a/Engine/Build/InstalledEngineBuild.xml
+++ b/Engine/Build/InstalledEngineBuild.xml
@@ -48,7 +48,7 @@
 	<Option Name="WithXboxOne" Restrict="true|false" DefaultValue="false" Description="Include the XboxOne target platform"/>
 	<Option Name="WithSwitch" Restrict="true|false" DefaultValue="false" Description="Include the Switch target platform"/>
 	<Option Name="WithLumin" Restrict="true|false" DefaultValue="$(DefaultWithPlatform)" Description="Include the Lumin target platform on Windows"/>
-	<Option Name="WithLuminMac" Restrict="true|false" DefaultValue="false" Description="Include the Lumin target platform on Mac"/>
+	<Option Name="WithLuminMac" Restrict="true|false" DefaultValue="$(DefaultWithPlatform)" Description="Include the Lumin target platform on Mac"/>
 	<Option Name="WithHoloLens" Restrict="true|false" DefaultValue="$(DefaultWithHoloLens)" Description="Include the HoloLens target platform"/>
 
 	<!-- Set up target types -->
