TP_ArchvizBP Should be TP_AEC_ArchvisBP in build graph too

diff --git a/Engine/Build/InstalledEngineFilters.xml b/Engine/Build/InstalledEngineFilters.xml
--- a/Engine/Build/InstalledEngineFilters.xml
+++ b/Engine/Build/InstalledEngineFilters.xml
@@ -173,9 +173,9 @@
 
 		<!-- Enterprise Templates -->
 		Templates/TP_AEC_BlankBP/...
+		Templates/TP_AEC_ArchvisBP/...
 		Templates/TP_MFG_BlankBP/...
 		Templates/TP_CollaborativeBP/...
-		Templates/TP_ArchvisBP/...
 		Templates/TP_PhotoStudioBP/...
 
 		Templates/TP_ME_BlankBP/...
