Add DatasmithCADWorker to binary package

#rnx

diff --git a/Engine/Build/InstalledEngineBuild.xml b/Engine/Build/InstalledEngineBuild.xml
--- a/Engine/Build/InstalledEngineBuild.xml
+++ b/Engine/Build/InstalledEngineBuild.xml
@@ -469,6 +469,7 @@
 		<Node Name="Build Tools Win64" Requires="Compile UnrealHeaderTool Win64">
 			<Compile Target="CrashReportClient" Configuration="Shipping" Platform="Win64" Tag="#Build Tools Win64"/>
 			<Compile Target="CrashReportClientEditor" Configuration="Shipping" Platform="Win64" Tag="#Build Tools Win64"/>
+			<Compile Target="DatasmithCADWorker" Configuration="Development" Platform="Win64" Tag="#Build Tools Win64"/>
 			<Compile Target="UnrealWatchdog" Configuration="Shipping" Platform="Win64" Tag="#Build Tools Win64"/>
 			<Compile Target="ShaderCompileWorker" Configuration="Development" Platform="Win64" Tag="#Build Tools Win64"/>
 			<Compile Target="UnrealCEFSubProcess" Configuration="Development" Platform="Win64" Tag="#Build Tools Win64"/>
