Fix errors making installed engine build with PS4 support.

#rnx

diff --git a/Engine/Build/InstalledEngineBuild.xml b/Engine/Build/InstalledEngineBuild.xml
--- a/Engine/Build/InstalledEngineBuild.xml
+++ b/Engine/Build/InstalledEngineBuild.xml
@@ -500,7 +500,7 @@
 			<!-- Copy Source and referenced libraries to a new location with Confidential folders removed-->
 			<Tag Files="Engine/Source/Programs/AutomationTool/..." Filter="*.csproj" Except="$(ConfidentialExceptions)"  With="#UAT Projects"/>
 			<CsCompile Project="#UAT Projects" Configuration="Development" Platform="AnyCPU" TagReferences="#UATReferences" EnumerateOnly="true"/>
-			<Tag Files="Engine/Binaries/DotNET/...;Engine/Binaries/ThirdParty/Newtonsoft/...;Engine/Binaries/ThirdParty/VisualStudio/...;Engine/Source/Programs/...;Engine/Source/Editor/SwarmInterface/...;#UATReferences" Except="$(ConfidentialExceptions);AutomationTool.csproj.References" With="#RedistUATSource"/>
+			<Tag Files="Engine/Binaries/DotNET/...;Engine/Binaries/ThirdParty/Newtonsoft/...;Engine/Binaries/ThirdParty/VisualStudio/...;Engine/Source/Programs/...;Engine/Platforms/*/Source/Programs/...;Engine/Source/Editor/SwarmInterface/...;#UATReferences" Except="$(ConfidentialExceptions);AutomationTool.csproj.References" With="#RedistUATSource"/>
 			<Copy Files="#RedistUATSource" From="$(RootDir)" To="$(CsToolsDir)"/>
 			<Tag Files="$(CsToolsDir)/Engine/Source/Programs/AutomationTool/..." Filter="*.csproj" With="#RedistUATProjects"/>
 
@@ -512,8 +512,8 @@
 				<CsCompile Project="$(CsToolsDir)/Engine/Source/Editor/SwarmInterface/DotNET/SwarmInterface.csproj" Configuration="Development" Platform="AnyCPU" Tag="#Build Tools CS Binaries;#Sign Binaries" TagReferences="#Build Tools CS Binaries"/>
 				<CsCompile Project="$(CsToolsDir)/Engine/Source/Programs/nDisplayLauncher/nDisplayLauncher.csproj" Configuration="Development" Platform="AnyCPU" Tag="#Build Tools CS Binaries;#Sign Binaries" TagReferences="#Build Tools CS Binaries"/>
 				<CsCompile Project="$(CsToolsDir)/Engine/Source/Programs/nDisplayListener/nDisplayListener.csproj" Configuration="Development" Platform="AnyCPU" Tag="#Build Tools CS Binaries;#Sign Binaries" TagReferences="#Build Tools CS Binaries"/>
-				<CsCompile Project="$(CsToolsDir)/Engine/Source/Programs/PS4/PS4DevKitUtil/PS4DevKitUtil.csproj" Configuration="Development" Platform="AnyCPU" Tag="#Build Tools CS Binaries;#Sign Binaries" TagReferences="#Build Tools CS Binaries" If="'$(WithPS4)' == true"/>
-				<CsCompile Project="$(CsToolsDir)/Engine/Source/Programs/PS4/PS4SymbolTool/PS4SymbolTool.csproj" Configuration="Release" Platform="AnyCPU" Tag="#Build Tools CS Binaries;#Sign Binaries" TagReferences="#Build Tools CS Binaries" If="'$(WithPS4)' == true"/>
+				<CsCompile Project="$(CsToolsDir)/Engine/Platforms/PS4/Source/Programs/PS4DevKitUtil/PS4DevKitUtil.csproj" Configuration="Development" Platform="AnyCPU" Tag="#Build Tools CS Binaries;#Sign Binaries" TagReferences="#Build Tools CS Binaries" If="'$(WithPS4)' == true"/>
+				<CsCompile Project="$(CsToolsDir)/Engine/Platforms/PS4/Source/Programs/PS4SymbolTool/PS4SymbolTool.csproj" Configuration="Release" Platform="AnyCPU" Tag="#Build Tools CS Binaries;#Sign Binaries" TagReferences="#Build Tools CS Binaries" If="'$(WithPS4)' == true"/>
 				<CsCompile Project="$(CsToolsDir)/Engine/Source/Programs/UnrealControls/UnrealControls.csproj" Configuration="Development" Platform="AnyCPU" Tag="#Build Tools CS Binaries;#Sign Binaries" TagReferences="#Build Tools CS Binaries"/>
 				<CsCompile Project="$(CsToolsDir)/Engine/Source/Programs/IOS/iPhonePackager/iPhonePackager.csproj" Configuration="Development" Platform="AnyCPU" Arguments="/verbosity:minimal /target:Rebuild" Tag="#Build Tools CS Binaries;#Sign Binaries" TagReferences="#Build Tools CS Binaries"/>
 			</Do>
