Fix support for building installed builds with console support.

#rnx

diff --git a/Engine/Build/InstalledEngineFilters.xml b/Engine/Build/InstalledEngineFilters.xml
--- a/Engine/Build/InstalledEngineFilters.xml
+++ b/Engine/Build/InstalledEngineFilters.xml
@@ -8,6 +8,9 @@
 			<Case If="'$(RestrictedFolderName)' == 'PS4' And '$(WithPS4)' == true">
 				<!-- Allow this folder -->
 			</Case>
+			<Case If="'$(RestrictedFolderName)' == 'Sony' And '$(WithPS4)' == true">
+				<!-- Allow this folder -->
+			</Case>
 			<Case If="'$(RestrictedFolderName)' == 'XboxOne' And '$(WithXboxOne)' == true">
 				<!-- Allow this folder -->
 			</Case>
