Added UnrealInsights to be distributed with Mac and Linux binaries.

diff --git a/Engine/Build/InstalledEngineBuild.xml b/Engine/Build/InstalledEngineBuild.xml
--- a/Engine/Build/InstalledEngineBuild.xml
+++ b/Engine/Build/InstalledEngineBuild.xml
@@ -537,6 +537,7 @@
 			<Compile Target="CrashReportClient" Platform="Linux" Configuration="Shipping"/>
 			<Compile Target="CrashReportClientEditor" Platform="Linux" Configuration="Shipping"/>
 			<Compile Target="ShaderCompileWorker" Configuration="Development" Platform="Linux" Tag="#Build Tools Linux"/>
+			<Compile Target="UnrealInsights" Configuration="Development" Platform="Linux" Tag="#Build Tools Linux"/>
 			<Compile Target="UnrealLightmass" Configuration="Development" Platform="Linux" Tag="#Build Tools Linux"/>
 			<Compile Target="UnrealPak" Configuration="Development" Platform="Linux" Tag="#Build Tools Linux"/>
 			<Compile Target="UnrealMultiUserServer" Configuration="Development" Platform="Linux" Tag="#Build Tools Linux"/>
@@ -555,6 +556,7 @@
 			<Compile Target="CrashReportClientEditor" Configuration="Shipping" Platform="Mac" Tag="#Build Tools Mac"/>
 			<Compile Target="ShaderCompileWorker" Configuration="Development" Platform="Mac" Tag="#Build Tools Mac"/>
 			<Compile Target="UnrealCEFSubProcess" Configuration="Development" Platform="Mac" Tag="#Build Tools Mac"/>
+			<Compile Target="UnrealInsights" Configuration="Development" Platform="Mac" Tag="#Build Tools Mac"/>
 			<Compile Target="UnrealFrontend" Configuration="Development" Platform="Mac" Tag="#Build Tools Mac"/>
 			<Compile Target="UnrealLightmass" Configuration="Development" Platform="Mac" Tag="#Build Tools Mac"/>
 			<Compile Target="UnrealPak" Configuration="Development" Platform="Mac" Tag="#Build Tools Mac"/>
