Generate DDC for all AEC and MFG templates

diff --git a/Engine/Build/InstalledEngineFilters.xml b/Engine/Build/InstalledEngineFilters.xml
--- a/Engine/Build/InstalledEngineFilters.xml
+++ b/Engine/Build/InstalledEngineFilters.xml
@@ -85,6 +85,9 @@
 		Templates/TP_2DSideScrollerBP/TP_2DSideScrollerBP.uproject
 		Templates/TP_VehicleAdvBP/TP_VehicleAdvBP.uproject
 		Templates/TP_CollaborativeBP/TP_CollaborativeBP.uproject
+		Templates/TP_AEC_BlankBP/TP_AEC_BlankBP.uproject
+		Templates/TP_AEC_ArchvisBP/TP_AEC_ArchvisBP.uproject
+		Templates/TP_PhotoStudioBP/TP_PhotoStudioBP.uproject
 		Samples/StarterContent/StarterContent.uproject
 		Samples/MobileStarterContent/MobileStarterContent.uproject
 	</Property>
@@ -177,7 +180,6 @@
 		<!-- Enterprise Templates -->
 		Templates/TP_AEC_BlankBP/...
 		Templates/TP_AEC_ArchvisBP/...
-		Templates/TP_MFG_BlankBP/...
 		Templates/TP_CollaborativeBP/...
 		Templates/TP_PhotoStudioBP/...
 
