Copy p4 editor binaries for Linux BuildGraph builds with WithFullDebugInfo set

Doing a BuildGraph build with -set:WithFullDebugInfo wasn't copying UE4Editor to the LocalBuilds directory.

./Engine/Build/BatchFiles/RunUAT.sh BuildGraph -script=Engine/Build/InstalledEngineBuild.xml -target="Make Installed Build Linux" -nosign -set:WithDDC=true -set:WithIOS=false -set:WithTVOS=false -set:WithMac=false -set:WithWin64=false -set:WithWin32=false -set:GameConfigurations=Development -set:WithFullDebugInfo=true -clean -verbose

This change sets "#UE4Editor Linux Stripped" to have the LocalBuild binary paths like non debug builds do.

diff --git a/Engine/Build/InstalledEngineBuild.xml b/Engine/Build/InstalledEngineBuild.xml
--- a/Engine/Build/InstalledEngineBuild.xml
+++ b/Engine/Build/InstalledEngineBuild.xml
@@ -195,7 +195,7 @@
 			<Switch>
 				<Case If="'$(WithFullDebugInfo)' == true">
 					<!-- Copy files for local debug builds -->
-					<Copy Files="#UE4Editor Linux Unstripped" From="$(RootDir)" To="$(SavedOutput)"/>
+					<Copy Files="#UE4Editor Linux Unstripped" From="$(RootDir)" To="$(SavedOutput)" Tag="#UE4Editor Linux Stripped"/>
 				</Case>
 				<Default>
 					<!-- Otherwise, strip the files for distribution -->
