Fix missing files when creating an installed build with PS4 support.

diff --git a/Engine/Build/InstalledEngineFilters.xml b/Engine/Build/InstalledEngineFilters.xml
--- a/Engine/Build/InstalledEngineFilters.xml
+++ b/Engine/Build/InstalledEngineFilters.xml
@@ -97,6 +97,7 @@
 		Templates/TP_VirtualRealityBP/TP_VirtualRealityBP.uproject
 	</Property>
 	-->
+	
 	<!-- Define Editor Filters -->
 	<Property Name="CopyEditorFilter">
 		<!-- This assembly is normally embedded into the UBT executable, but it can technically be rebuilt from an installed build -->
@@ -110,39 +111,9 @@
 		<!-- Content folders -->
 		Engine/Content/...
 
-		<!-- Config files -->
-		Engine/Config/...
-
-		<!-- Programs -->
-		Engine/Programs/...
-
-		<!-- Plugins -->
-		Engine/Plugins/....uplugin
-		Engine/Plugins/.../Content/...
-		Engine/Plugins/.../Config/...
-		Engine/Plugins/.../Resources/...
-		Engine/Plugins/.../Shaders/...
-		Engine/Plugins/.../Source/...
-		Engine/Plugins/.../Templates/...
-
 		<!-- Source code -->
 		Engine/Source/UE4Game.Target.cs
 		Engine/Source/UE4Editor.Target.cs
-		Engine/Source/Runtime/...
-		Engine/Source/Developer/...
-		Engine/Source/Editor/...
-		Engine/Source/ThirdParty/Licenses/...
-		Engine/Source/Programs/AutomationTool/...
-		Engine/Source/Programs/AutomationToolLauncher/...
-		Engine/Source/Programs/DotNETCommon/...
-		Engine/Source/Programs/UnrealBuildTool/...
-		Engine/Source/Programs/UnrealHeaderTool/...
-
-		<!-- Third Party Software description files -->
-		Engine/.../*.tps
-
-		<!-- Shaders -->
-		Engine/Shaders/...
 
 		<!-- Starter content -->
 		Samples/StarterContent/Content/...
@@ -193,6 +164,44 @@
 		Engine/Build/Build.version
 		Engine/Build/Target.cs.template
 	</Property>
+	
+	<!-- Files which can exist under any engine or engine-platform directory -->
+	<Property Name="CopyEditorEngineOrPlatformFilter" Value="">
+		<!-- Config files -->
+		Config/...
+
+		<!-- Programs -->
+		Programs/...
+
+		<!-- Plugins -->
+		Plugins/....uplugin
+		Plugins/.../Content/...
+		Plugins/.../Config/...
+		Plugins/.../Resources/...
+		Plugins/.../Shaders/...
+		Plugins/.../Source/...
+		Plugins/.../Templates/...
+
+		<!-- Source code -->
+		Source/Runtime/...
+		Source/Developer/...
+		Source/Editor/...
+		Source/ThirdParty/Licenses/...
+		Source/Programs/AutomationTool/...
+		Source/Programs/AutomationToolLauncher/...
+		Source/Programs/DotNETCommon/...
+		Source/Programs/UnrealBuildTool/...
+		Source/Programs/UnrealHeaderTool/...
+
+		<!-- Third Party Software description files -->
+		.../*.tps
+
+		<!-- Shaders -->
+		Shaders/...
+	</Property>
+	<ForEach Name="Filter" Values="$(CopyEditorEngineOrPlatformFilter)">
+		<Property Name="CopyEditorFilter" Value="$(CopyEditorFilter);Engine/$(Filter)"/>
+	</ForEach>
 	<!-- Optional target files -->
 	<Property Name="CopyEditorFilter" Value="$(CopyEditorFilter);Engine/Source/UE4Client.Target.cs" If="$(WithClient)"/>
 	<Property Name="CopyEditorFilter" Value="$(CopyEditorFilter);Engine/Source/UE4Server.Target.cs" If="$(WithServer)"/>
@@ -481,7 +490,7 @@
 		Engine/Build/TVOS/...
 	</Property>
 
-	<!-- Win64 specific -->
+	<!-- TVOS on Win64 specific -->
 	<Do If="'$(HostPlatform)' == 'Win64'">
 		<Property Name="CopyTVOSFilter" Value="$(CopyTVOSFilter)">
 			Engine/Binaries/DotNET/IOS/openssl.exe
@@ -514,11 +523,14 @@
 	<Property Name="CopyLinuxExceptions">
 	</Property>
 
-<Property Name="CopyLinuxFilter" Value="$(CopyLinuxFilter);Engine/Extras/ThirdPartyNotUE/putty/..."/>	
+	<Property Name="CopyLinuxFilter" Value="$(CopyLinuxFilter);Engine/Extras/ThirdPartyNotUE/putty/..."/>	
 
-<!-- Win64 specific -->
+	<!-- Win64 specific -->
 
 	<!-- PS4 -->
+	<ForEach Name="Filter" Values="$(CopyEditorEngineOrPlatformFilter)">
+		<Property Name="CopyEditorFilter" Value="$(CopyEditorFilter);Platforms/Sony/$(Filter);Platforms/PS4/$(Filter)"/>
+	</ForEach>
 	<Property Name="CopyPS4Filter">
 		<!-- Build Files -->
 		Engine/Build/PS4/...
@@ -534,12 +546,11 @@
 
 		<!-- Additional Binaries -->
 		Engine/Binaries/Win64/XboxOne/XboxOnePackageNameUtil.exe
-    Engine/Binaries/Win64/XboxOne/XboxOnePDBFileUtil.exe
+		Engine/Binaries/Win64/XboxOne/XboxOnePDBFileUtil.exe
 	</Property>
 
 	<Property Name="CopyXboxOneExceptions">
-	
-</Property>
+	</Property>
 
 	<!-- Switch -->
 	<Property Name="CopySwitchFilter">
