Fix missing FP_FirstPerson files in binary distribution.

#rnx

diff --git a/Engine/Build/InstalledEngineFilters.xml b/Engine/Build/InstalledEngineFilters.xml
--- a/Engine/Build/InstalledEngineFilters.xml
+++ b/Engine/Build/InstalledEngineFilters.xml
@@ -121,6 +121,8 @@
 
 		<!-- Templates -->
 		Templates/TemplateCategories.ini
+		Templates/FP_FirstPerson/...
+		Templates/FP_FirstPersonBP/...
 		Templates/Media/...
 		Templates/TP_Blank/...
 		Templates/TP_BlankBP/...
