Fix platform config files not being included when making an installed build with PS4 support.

#rnx

diff --git a/Engine/Build/InstalledEngineFilters.xml b/Engine/Build/InstalledEngineFilters.xml
--- a/Engine/Build/InstalledEngineFilters.xml
+++ b/Engine/Build/InstalledEngineFilters.xml
@@ -531,7 +531,7 @@
 
 	<!-- PS4 -->
 	<ForEach Name="Filter" Values="$(CopyEditorEngineOrPlatformFilter)">
-		<Property Name="CopyEditorFilter" Value="$(CopyEditorFilter);Platforms/Sony/$(Filter);Platforms/PS4/$(Filter)"/>
+		<Property Name="CopyEditorFilter" Value="$(CopyEditorFilter);Engine/Platforms/Sony/$(Filter);Engine/Platforms/PS4/$(Filter)"/>
 	</ForEach>
 	<Property Name="CopyPS4Filter">
 		<!-- Build Files -->
