Copy AutomationTool projects for platform plugins in installed builds.

#rnx

diff --git a/Engine/Build/InstalledEngineBuild.xml b/Engine/Build/InstalledEngineBuild.xml
--- a/Engine/Build/InstalledEngineBuild.xml
+++ b/Engine/Build/InstalledEngineBuild.xml
@@ -498,11 +498,16 @@
 			<Delete Files="$(CsToolsDir)/..."/>
 
 			<!-- Copy Source and referenced libraries to a new location with Confidential folders removed-->
-			<Tag Files="Engine/Source/Programs/AutomationTool/..." Filter="*.csproj" Except="$(ConfidentialExceptions)"  With="#UAT Projects"/>
+			<Tag Files="$(AutomationToolPaths)" Filter="*.csproj" Except="$(ConfidentialExceptions)"  With="#UAT Projects"/>
 			<CsCompile Project="#UAT Projects" Configuration="Development" Platform="AnyCPU" TagReferences="#UATReferences" EnumerateOnly="true"/>
 			<Tag Files="Engine/Binaries/DotNET/...;Engine/Binaries/ThirdParty/Newtonsoft/...;Engine/Binaries/ThirdParty/VisualStudio/...;Engine/Source/Programs/...;Engine/Platforms/*/Source/Programs/...;Engine/Source/Editor/SwarmInterface/...;#UATReferences" Except="$(ConfidentialExceptions);AutomationTool.csproj.References" With="#RedistUATSource"/>
 			<Copy Files="#RedistUATSource" From="$(RootDir)" To="$(CsToolsDir)"/>
-			<Tag Files="$(CsToolsDir)/Engine/Source/Programs/AutomationTool/..." Filter="*.csproj" With="#RedistUATProjects"/>
+
+			<Property Name="RedistAutomationToolPaths" Value=""/>
+			<ForEach Name="AutomationToolPath" Values="$(AutomationToolPaths)">
+				<Property Name="RedistAutomationToolPaths" Value="$(RedistAutomationToolPaths)$(CsToolsDir)/$(AutomationToolPath);"/>
+			</ForEach>
+			<Tag Files="$(RedistAutomationToolPaths)" Filter="*.csproj" With="#RedistUATProjects"/>
 
 			<!-- Compile all the tools -->
 			<CsCompile Project="#RedistUATProjects" Configuration="Development" Platform="AnyCPU" Tag="#Build Tools CS Binaries;#Sign Binaries" TagReferences="#Build Tools CS Binaries"/>
diff --git a/Engine/Build/InstalledEngineFilters.xml b/Engine/Build/InstalledEngineFilters.xml
--- a/Engine/Build/InstalledEngineFilters.xml
+++ b/Engine/Build/InstalledEngineFilters.xml
@@ -251,6 +251,9 @@
 	<!-- exclude VR feature pack on Mac -->
 	<Property Name="CopyEditorExceptions" Value="$(CopyEditorExceptions);FeaturePacks/TP_VirtualRealityBP.upack" If="'$(HostPlatform)' == 'Linux'"/>
 
+	<!-- Projects to be compiled for UAT -->
+	<Property Name="AutomationToolPaths" Value="Engine/Source/Programs/AutomationTool/..."/>
+	
 	<!-- Win64-specific dependencies -->
 	<Do If="'$(HostPlatform)' == 'Win64'">
 		<Property Name="CopyEditorFilter" Value="$(CopyEditorFilter)">
@@ -530,9 +533,12 @@
 	<!-- Win64 specific -->
 
 	<!-- PS4 -->
-	<ForEach Name="Filter" Values="$(CopyEditorEngineOrPlatformFilter)">
-		<Property Name="CopyEditorFilter" Value="$(CopyEditorFilter);Engine/Platforms/Sony/$(Filter);Engine/Platforms/PS4/$(Filter)"/>
-	</ForEach>
+	<Do If="$(WithPS4)">
+		<ForEach Name="Filter" Values="$(CopyEditorEngineOrPlatformFilter)">
+			<Property Name="CopyEditorFilter" Value="$(CopyEditorFilter);Engine/Platforms/Sony/$(Filter);Engine/Platforms/PS4/$(Filter)"/>
+		</ForEach>
+		<Property Name="AutomationToolPaths" Value="$(AutomationToolPaths);Engine/Platforms/Sony/Source/Programs/AutomationTool/...;Engine/Platforms/PS4/Source/Programs/AutomationTool/..."/>
+	</Do>
 	<Property Name="CopyPS4Filter">
 		<!-- Build Files -->
 		Engine/Build/PS4/...
