Changes to support building precompiled binaries using UE4_Licensee.

#rnx

diff --git a/Engine/Build/InstalledEngineBuild.xml b/Engine/Build/InstalledEngineBuild.xml
--- a/Engine/Build/InstalledEngineBuild.xml
+++ b/Engine/Build/InstalledEngineBuild.xml
@@ -117,7 +117,7 @@
 	<Include Script="InstalledEngineFilters.xml"/>
 
 	<!-- Base editor and tools for compile on Windows -->
-	<Agent Name="Editor Win64" Type="CompileWin64;Win64">
+	<Agent Name="Editor Win64" Type="Win64_Licensee">
 		<Node Name="Update Version Files">
 			<Log Message="$(RootDir)"/>
 			<SetVersion Change="$(Change)" Branch="$(EscapedBranch)" If="$(IsBuildMachine)"/>
@@ -160,7 +160,7 @@
 
 
 	<!-- Base editor and tools for compile on Mac -->
-	<Agent Name="Editor Mac" Type="CompileMac;Mac">
+	<Agent Name="Editor Mac" Type="Mac_Licensee">
 		<Node Name="Compile UnrealHeaderTool Mac" Requires="Update Version Files">
 			<Compile Target="UnrealHeaderTool" Platform="Mac" Configuration="Development" Arguments="-precompile -allmodules"/>
 		</Node>
@@ -190,7 +190,7 @@
 	</Agent>
 
 	<!-- Base editor and tools for compile on Linux -->
-	<Agent Name="Editor Linux" Type="CompileWin64;Win64">
+	<Agent Name="Editor Linux" Type="Win64_Licensee">
 		<Node Name="Compile UnrealHeaderTool Linux" Requires="Update Version Files">
 			<Compile Target="UnrealHeaderTool" Platform="Linux" Configuration="Development" Arguments="-precompile -allmodules"/>
 		</Node>
@@ -214,7 +214,7 @@
 	</Agent>
 
 	<!-- Target Platforms normally built from a Win64 machine -->
-	<Agent Name="Target Win64" Type="CompileWin64;Win64">
+	<Agent Name="Target Win64" Type="Win64_Licensee">
 		<Node Name="Compile UE4Game Win64" Requires="Compile UnrealHeaderTool Win64" Produces="#UE4Game Win64;#UE4Game Win64 Unstripped;#UE4Game Win64 Stripped;#UE4Game Win64 Unsigned;#UE4Game Win64 Signed">
 
 			<ForEach Name="Target" Values="UE4Game;$(OptionalClientTarget);$(OptionalServerTarget)">
@@ -251,7 +251,7 @@
 			</Do>
 		</Node>
 	</Agent>
-	<Agent Name="Target Win32" Type="CompileWin64;Win64">
+	<Agent Name="Target Win32" Type="Win64_Licensee">
 		<Node Name="Compile UE4Game Win32" Requires="Compile UnrealHeaderTool Win64" Produces="#UE4Game Win32;#UE4Game Win32 Unstripped;#UE4Game Win32 Stripped;#UE4Game Win32 Unsigned;#UE4Game Win32 Signed">
 
 			<ForEach Name="Target" Values="UE4Game;$(OptionalClientTarget)">
@@ -289,7 +289,7 @@
 			</Do>
 		</Node>
 	</Agent>
-	<Agent Name="Target Linux" Type="CompileWin64;Win64">
+	<Agent Name="Target Linux" Type="Win64_Licensee">
 		<Node Name="Compile UE4Game Linux" Requires="Compile UnrealHeaderTool $(HostPlatform)" Produces="#UE4Game Linux;#UE4Game Linux Unstripped;#UE4Game Linux Stripped">
 
 			<ForEach Name="Target" Values="UE4Game;$(OptionalClientTarget);$(OptionalServerTarget)">
@@ -304,7 +304,7 @@
 
 		</Node>
 	</Agent>
-	<Agent Name="Target LinuxAArch64" Type="CompileWin64;Win64">
+	<Agent Name="Target LinuxAArch64" Type="Win64_Licensee">
 		<Node Name="Compile UE4Game LinuxAArch64" Requires="Compile UnrealHeaderTool $(HostPlatform)" Produces="#UE4Game LinuxAArch64;#UE4Game LinuxAArch64 Unstripped;#UE4Game LinuxAArch64 Stripped">
 
 			<ForEach Name="Target" Values="UE4Game;$(OptionalClientTarget);$(OptionalServerTarget)">
@@ -319,7 +319,7 @@
 
 		</Node>
 	</Agent>
-	<Agent Name="Target Android" Type="CompileWin64;Win64">
+	<Agent Name="Target Android" Type="Win64_Licensee">
 		<Node Name="Compile UE4Game Android" Requires="Compile UnrealHeaderTool $(HostPlatform)" Produces="#UE4Game Android;#UE4Game Android Unstripped;#UE4Game Android Stripped">
 
 			<ForEach Name="Target" Values="UE4Game;$(OptionalClientTarget)">
@@ -334,7 +334,7 @@
 
 		</Node>
 	</Agent>
-	<Agent Name="Target PS4" Type="CompileWin64;Win64">
+	<Agent Name="Target PS4" Type="Win64_Licensee">
 		<Node Name="Compile UE4Game PS4" Requires="Compile UnrealHeaderTool Win64" Produces="#UE4Game PS4">
 
 			<ForEach Name="Target" Values="UE4Game;$(OptionalClientTarget)">
@@ -346,7 +346,7 @@
 
 		</Node>
 	</Agent>
-	<Agent Name="Target XboxOne" Type="CompileWin64;Win64">
+	<Agent Name="Target XboxOne" Type="Win64_Licensee">
 		<Node Name="Compile UE4Game XboxOne" Requires="Compile UnrealHeaderTool Win64" Produces="#UE4Game XboxOne">
 
 			<ForEach Name="Target" Values="UE4Game;$(OptionalClientTarget)">
@@ -358,7 +358,7 @@
 
 		</Node>
 	</Agent>
-	<Agent Name="Target Switch" Type="CompileWin64;Win64">
+	<Agent Name="Target Switch" Type="Win64_Licensee">
 		<Node Name="Compile UE4Game Switch" Requires="Compile UnrealHeaderTool Win64" Produces="#UE4Game Switch" If="'$(WithSwitch)' == true">
 
 			<ForEach Name="Target" Values="UE4Game;$(OptionalClientTarget)">
@@ -381,7 +381,7 @@
 			</ForEach>
 		</Node>
 	</Agent>
-	<Agent Name="Target Lumin" Type="CompileWin64;Win64">
+	<Agent Name="Target Lumin" Type="Win64_Licensee">
 		<Node Name="Compile UE4Game Lumin" Requires="Compile UnrealHeaderTool Win64" Produces="#UE4Game Lumin;#UE4Game Lumin Unstripped;#UE4Game Lumin Stripped" If="'$(WithLumin)' == true">
 
 			<ForEach Name="Target" Values="UE4Game;$(OptionalClientTarget)">
@@ -398,7 +398,7 @@
 	</Agent>
 
 	<!-- Target Platforms normally built from a Mac machine -->
-	<Agent Name="Target Mac" Type="CompileMac;Mac">
+	<Agent Name="Target Mac" Type="Mac_Licensee">
 		<Node Name="Compile UE4Game Mac" Requires="Compile UnrealHeaderTool Mac" Produces="#UE4Game Mac;#UE4Game Mac Unstripped;#UE4Game Mac Stripped;#UE4Game Mac Unsigned;#UE4Game Mac Signed">
 
 			<ForEach Name="Target" Values="UE4Game;$(OptionalClientTarget);$(OptionalServerTarget)">
@@ -419,7 +419,7 @@
 
 		</Node>
 	</Agent>
-	<Agent Name="Target IOS" Type="CompileMac;Mac">
+	<Agent Name="Target IOS" Type="Mac_Licensee">
 		<Node Name="Compile UE4Game IOS" Requires="Compile UnrealHeaderTool Mac" Produces="#UE4Game IOS;#UE4Game IOS Unstripped;#UE4Game IOS Stripped">
 
 			<ForEach Name="Target" Values="UE4Game;$(OptionalClientTarget)">
@@ -434,7 +434,7 @@
 
 		</Node>
 	</Agent>
-	<Agent Name="Target TVOS" Type="CompileMac;Mac">
+	<Agent Name="Target TVOS" Type="Mac_Licensee">
 		<Node Name="Compile UE4Game TVOS" Requires="Compile UnrealHeaderTool Mac" Produces="#UE4Game TVOS;#UE4Game TVOS Unstripped;#UE4Game TVOS Stripped">
 
 			<ForEach Name="Target" Values="UE4Game;$(OptionalClientTarget)">
@@ -449,7 +449,7 @@
 
 		</Node>
 	</Agent>
-	<Agent Name="Target LuminMac" Type="CompileMac;Mac">
+	<Agent Name="Target LuminMac" Type="Mac_Licensee">
 		<Node Name="Compile UE4Game LuminMac" Requires="Compile UnrealHeaderTool Mac" Produces="#UE4Game LuminMac;#UE4Game LuminMac Unstripped;#UE4Game LuminMac Stripped">
 
 			<ForEach Name="Target" Values="UE4Game;$(OptionalClientTarget)">
@@ -465,7 +465,7 @@
 		</Node>
 	</Agent>
 
-	<Agent Name="Tools Group Win64" Type="CompileWin64;Win64">
+	<Agent Name="Tools Group Win64" Type="Win64_Licensee">
 		<Node Name="Build Tools Win64" Requires="Compile UnrealHeaderTool Win64">
 			<Compile Target="CrashReportClient" Configuration="Shipping" Platform="Win64" Tag="#Build Tools Win64"/>
 			<Compile Target="CrashReportClientEditor" Configuration="Shipping" Platform="Win64" Tag="#Build Tools Win64"/>
@@ -550,7 +550,7 @@
 		</Node>
 	</Agent>
 
-	<Agent Name="Tools Group Mac" Type="CompileMac;Mac">
+	<Agent Name="Tools Group Mac" Type="Mac_Licensee">
 		<Node Name="Build Tools Mac" Requires="Compile UnrealHeaderTool Mac">
 			<Compile Target="CrashReportClient" Configuration="Shipping" Platform="Mac" Tag="#Build Tools Mac"/>
 			<Compile Target="CrashReportClientEditor" Configuration="Shipping" Platform="Mac" Tag="#Build Tools Mac"/>
@@ -569,7 +569,7 @@
 		</Node>
 	</Agent>
 
-	<Agent Name="FeaturePacks Group" Type="CompileWin64;Win64">
+	<Agent Name="FeaturePacks Group" Type="Win64_Licensee">
 		<Node Name="Make Feature Packs" Requires="Update Version Files;Build Tools $(HostPlatform)">
 			<!--Delete any pre-existing feature packs-->
 			<Delete Files="FeaturePacks/..."/>
@@ -597,7 +597,7 @@
 		<Property Name="DDCPlatformsWin64" Value="$(DDCPlatformsWin64)+HoloLens" If="'$(WithHoloLens)' == true"/>
 	</Do>
 
-	<Agent Name="DDC Group Win64" Type="Win64">
+	<Agent Name="DDC Group Win64" Type="Win64_Licensee">
 		<Node Name="Build DDC Win64" Requires="#UE4Editor Win64;Build Tools Win64;Build Tools CS" Produces="#CompressedDDCWin64">
 			<!-- Get our temp dir -->
 			<Property Name="DDCDir" Value="$(RootDir)\LocalBuilds\InstalledDDC"/>
@@ -619,12 +619,7 @@
 			<Tag Files="Templates\TemplateResources\..." With="#ToCopy"/>
 			<!-- Filter out the files not needed to build DDC. Removing confidential folders can affect DDC keys, so we want to be sure that we're making DDC with a build that can use it. -->
 			<Tag Files="#ToCopy" Except="$(ConfidentialExceptions);$(RootDir)/.../Source/...;$(RootDir)/.../Intermediate/..." With="#FilteredCopyList"/>
-			<!-- Additional required files that would be removed by confidential filter. Note that conditions are evaluated before copying build products, so check against EpicInternal.txt -->
-			<Tag Files="Engine\Build\NotForLicensees\EpicInternal.txt" With="#FilteredCopyList" If="Exists('Engine\Build\NotForLicensees\EpicInternal.txt')"/>
-			<Tag Files="Engine\Binaries\Win64\NotForLicensees\ShaderCompileWorker.modules" With="#FilteredCopyList" If="Exists('Engine\Build\NotForLicensees\EpicInternal.txt')"/>
-			<Tag Files="Engine\Binaries\Win64\NotForLicensees\ShaderCompileWorker-DDCUtils.dll" With="#FilteredCopyList" If="Exists('Engine\Build\NotForLicensees\EpicInternal.txt')"/>
-			<Tag Files="Engine\Binaries\Win64\NotForLicensees\UE4Editor.modules" With="#FilteredCopyList" If="Exists('Engine\Build\NotForLicensees\EpicInternal.txt')"/>
-			<Tag Files="Engine\Binaries\Win64\NotForLicensees\UE4Editor-DDCUtils.dll" With="#FilteredCopyList" If="Exists('Engine\Build\NotForLicensees\EpicInternal.txt')"/>
+			
 			<!-- Copy everything to a temporary directory -->
 			<Copy From="$(RootDir)" To="$(DDCDir)" Files="#FilteredCopyList"/>
 			<Copy From="$(CsToolsDir)" To="$(DDCDir)" Files="#Build Tools CS Binaries"/>
@@ -665,12 +660,7 @@
 			<Tag Files="Templates/TemplateResources/..." With="#ToCopy"/>
 			<!-- Filter out the files not needed to build DDC. Removing confidential folders can affect DDC keys, so we want to be sure that we're making DDC with a build that can use it. -->
 			<Tag Files="#ToCopy" Except="$(ConfidentialExceptions);$(RootDir)/.../Source/...;$(RootDir)/.../Intermediate/..." With="#FilteredCopyList"/>
-			<!-- Additional required files that would be removed by confidential filter -->
-			<Tag Files="Engine/Build/NotForLicensees/EpicInternal.txt" With="#FilteredCopyList" If="Exists('Engine/Build/NotForLicensees/EpicInternal.txt')"/>
-			<Tag Files="Engine/Binaries/Mac/NotForLicensees/ShaderCompileWorker.modules" With="#FilteredCopyList" If="Exists('Engine/Build/NotForLicensees/EpicInternal.txt')"/>
-			<Tag Files="Engine/Binaries/Mac/NotForLicensees/ShaderCompileWorker-DDCUtils.dylib" With="#FilteredCopyList" If="Exists('Engine/Build/NotForLicensees/EpicInternal.txt')"/>
-			<Tag Files="Engine/Binaries/Mac/NotForLicensees/UE4Editor.modules" With="#FilteredCopyList" If="Exists('Engine/Build/NotForLicensees/EpicInternal.txt')"/>
-			<Tag Files="Engine/Binaries/Mac/NotForLicensees/UE4Editor-DDCUtils.dylib" With="#FilteredCopyList" If="Exists('Engine/Build/NotForLicensees/EpicInternal.txt')"/>
+			
 			<!-- Copy everything to a temporary directory -->
 			<Copy From="$(RootDir)" To="$(DDCDir)" Files="#FilteredCopyList"/>
 			<Copy From="$(CsToolsDir)" To="$(DDCDir)" Files="#Build Tools CS Binaries"/>
@@ -713,18 +703,7 @@
 			<Tag Files="Templates/TemplateResources/..." With="#ToCopy"/>
 			<!-- Filter out the files not needed to build DDC. Removing confidential folders can affect DDC keys, so we want to be sure that we're making DDC with a build that can use it. -->
 			<Tag Files="#ToCopy" Except="$(ConfidentialExceptions);$(RootDir)/.../Source/...;$(RootDir)/.../Intermediate/..." With="#FilteredCopyList"/>
-			<!-- Additional required files that would be removed by confidential filter -->
-			<Tag Files="Engine/Build/NotForLicensees/EpicInternal.txt" With="#FilteredCopyList" If="Exists('Engine/Build/NotForLicensees/EpicInternal.txt')"/>
-			<!-- When running on Linux -->
-			<Tag Files="Engine/Binaries/Linux/NotForLicensees/ShaderCompileWorker.modules" With="#FilteredCopyList" If="Exists('Engine/Build/NotForLicensees/EpicInternal.txt')"/>
-			<Tag Files="Engine/Binaries/Linux/NotForLicensees/libShaderCompileWorker-DDCUtils.so" With="#FilteredCopyList" If="Exists('Engine/Build/NotForLicensees/EpicInternal.txt')"/>
-			<Tag Files="Engine/Binaries/Linux/NotForLicensees/UE4Editor.modules" With="#FilteredCopyList" If="Exists('Engine/Build/NotForLicensees/EpicInternal.txt')"/>
-			<Tag Files="Engine/Binaries/Linux/NotForLicensees/libUE4Editor-DDCUtils.so" With="#FilteredCopyList" If="Exists('Engine/Build/NotForLicensees/EpicInternal.txt')"/>
-			<!-- When running on Windows -->
-			<Tag Files="Engine/Binaries/Win64/NotForLicensees/ShaderCompileWorker.modules" With="#FilteredCopyList" If="Exists('Engine/Build/NotForLicensees/EpicInternal.txt') And '$(HostPlatform)' == 'Win64'"/>
-			<Tag Files="Engine/Binaries/Win64/NotForLicensees/ShaderCompileWorker-DDCUtils.dll" With="#FilteredCopyList" If="Exists('Engine/Build/NotForLicensees/EpicInternal.txt') And '$(HostPlatform)' == 'Win64'"/>
-			<Tag Files="Engine/Binaries/Win64/NotForLicensees/UE4Editor.modules" With="#FilteredCopyList" If="Exists('Engine/Build/NotForLicensees/EpicInternal.txt') And '$(HostPlatform)' == 'Win64'"/>
-			<Tag Files="Engine/Binaries/Win64/NotForLicensees/UE4Editor-DDCUtils.dll" With="#FilteredCopyList" If="Exists('Engine/Build/NotForLicensees/EpicInternal.txt') And '$(HostPlatform)' == 'Win64'"/>
+			
 			<!-- Copy everything to a temporary directory -->
 			<Copy From="$(RootDir)" To="$(DDCDir)" Files="#FilteredCopyList"/>
 			<Copy From="$(CsToolsDir)" To="$(DDCDir)" Files="#Build Tools CS Binaries"/>
@@ -735,7 +714,7 @@
 		</Node>
 	</Agent>
 
-	<Agent Name="Installed Build Group Win64" Type="Win64">
+	<Agent Name="Installed Build Group Win64" Type="Win64_Licensee">
 		<!-- Build up a list of requirements as some can be optional -->
 		<Property Name="InstalledRequirements" Value="Compile UnrealHeaderTool Win64;#UE4Editor Win64;#UE4Editor Win64 Unstripped;#UE4Editor Win64 Stripped;#UE4Editor Win64 Unsigned;#UE4Editor Win64 Signed;#Build Tools Win64;#Build Tools CS Binaries"/>
 		<!-- Optional Target Platform requirements -->
@@ -903,7 +882,7 @@
 		</Node>
 	</Agent>
 
-	<Agent Name="Installed Build Group Mac" Type="Mac">
+	<Agent Name="Installed Build Group Mac" Type="Mac_Licensee">
 		<!-- Build up a list of requirements as some can be optional -->
 		<Property Name="FilterRequirements" Value="#Compile UnrealHeaderTool Mac;#UE4Editor Mac;#UE4Editor Mac Unstripped;#UE4Editor Mac Stripped;#UE4Editor Mac Unsigned;#UE4Editor Mac Signed;#Build Tools Mac;#Build Tools CS Binaries"/>
 
@@ -1022,7 +1001,7 @@
 		</Node>
 	</Agent>
 
-	<Agent Name="Installed Build Group Linux" Type="Win64">
+	<Agent Name="Installed Build Group Linux" Type="Win64_Licensee">
 		<!-- Build up a list of requirements as some can be optional -->
 		<Property Name="InstalledRequirements" Value="#Compile UnrealHeaderTool Linux;#UE4Editor Linux;#UE4Editor Linux Unstripped;#UE4Editor Linux Stripped;#Build Tools Linux;#Build Tools CS Binaries"/>
 
diff --git a/Engine/Source/Programs/AutomationTool/BuildGraph/BuildGraph.cs b/Engine/Source/Programs/AutomationTool/BuildGraph/BuildGraph.cs
--- a/Engine/Source/Programs/AutomationTool/BuildGraph/BuildGraph.cs
+++ b/Engine/Source/Programs/AutomationTool/BuildGraph/BuildGraph.cs
@@ -219,7 +219,7 @@ namespace AutomationTool
 
 			// Read the script from disk
 			Graph Graph;
-			if(!ScriptReader.TryRead(new FileReference(ScriptFileName), Arguments, DefaultProperties, PreprocessedFileName != null, Schema, out Graph))
+			if(!ScriptReader.TryRead(new FileReference(ScriptFileName), Arguments, DefaultProperties, PreprocessedFileName != null, Schema, out Graph, SingleNodeName))
 			{
 				return ExitCode.Error_Unknown;
 			}
@@ -255,7 +255,20 @@ namespace AutomationTool
 			}
 			else
 			{
-				foreach(string TargetName in TargetNames.Split(new char[]{ '+', ';' }, StringSplitOptions.RemoveEmptyEntries).Select(x => x.Trim()))
+				IEnumerable<string> NodesToResolve = null;
+
+				// If we're only building a single node and using a preprocessed reference we only need to try to resolve the references
+				// for that node.
+				if (SingleNodeName != null && PreprocessedFileName != null)
+				{
+					NodesToResolve = new List<string> { SingleNodeName };
+				}
+				else
+				{
+					NodesToResolve = TargetNames.Split(new char[] { '+', ';' }, StringSplitOptions.RemoveEmptyEntries).Select(x => x.Trim());
+				}
+
+				foreach (string TargetName in NodesToResolve)
 				{
 					Node[] Nodes;
 					if(!Graph.TryResolveReference(TargetName, out Nodes))
@@ -658,7 +671,7 @@ namespace AutomationTool
 				if(!Storage.IsComplete(NodeToExecute.Name))
 				{
 					LogInformation("");
-					if(!BuildNode(Job, Graph, NodeToExecute, Storage, false))
+					if(!BuildNode(Job, Graph, NodeToExecute, Storage, bWithBanner: false))
 					{
 						return false;
 					} 
diff --git a/Engine/Source/Programs/AutomationTool/BuildGraph/Script.cs b/Engine/Source/Programs/AutomationTool/BuildGraph/Script.cs
--- a/Engine/Source/Programs/AutomationTool/BuildGraph/Script.cs
+++ b/Engine/Source/Programs/AutomationTool/BuildGraph/Script.cs
@@ -246,6 +246,11 @@ namespace AutomationTool
 		/// </summary>
 		int NumErrors;
 
+		/// <summary>
+		/// The name of the node if only a single node is going to be built, otherwise null.
+		/// </summary>
+		string SingleNodeName;
+
 		/// <summary>
 		/// Private constructor. Use ScriptReader.TryRead() to read a script file.
 		/// </summary>
@@ -274,8 +279,9 @@ namespace AutomationTool
 		/// <param name="bPreprocessOnly">Preprocess the file, but do not expand any values that are not portable (eg. paths on the local machine)</param>
 		/// <param name="Schema">Schema for the script</param>
 		/// <param name="Graph">If successful, the graph constructed from the given script</param>
+		/// <param name="SingleNodeName">If a single node will be processed, the name of that node.</param>
 		/// <returns>True if the graph was read, false if there were errors</returns>
-		public static bool TryRead(FileReference File, Dictionary<string, string> Arguments, Dictionary<string, string> DefaultProperties, bool bPreprocessOnly, ScriptSchema Schema, out Graph Graph)
+		public static bool TryRead(FileReference File, Dictionary<string, string> Arguments, Dictionary<string, string> DefaultProperties, bool bPreprocessOnly, ScriptSchema Schema, out Graph Graph, string SingleNodeName = null)
 		{
 			// Check the file exists before doing anything.
 			if (!FileReference.Exists(File))
@@ -287,7 +293,7 @@ namespace AutomationTool
 
 			// Read the file and build the graph
 			ScriptReader Reader = new ScriptReader(DefaultProperties, bPreprocessOnly, Schema);
-			if (!Reader.TryRead(File, Arguments) || Reader.NumErrors > 0)
+			if (!Reader.TryRead(File, Arguments, SingleNodeName) || Reader.NumErrors > 0)
 			{
 				Graph = null;
 				return false;
@@ -316,7 +322,8 @@ namespace AutomationTool
 		/// </summary>
 		/// <param name="File">File to read from</param>
 		/// <param name="Arguments">Arguments passed in to the graph on the command line</param>
-		bool TryRead(FileReference File, Dictionary<string, string> Arguments)
+		/// <param name="SingleNodeName">The name of the node if only a single node is going to be built, otherwise null.</param>
+		bool TryRead(FileReference File, Dictionary<string, string> Arguments, string SingleNodeName = null)
 		{
 			// Read the document and validate it against the schema
 			ScriptDocument Document;
@@ -327,6 +334,7 @@ namespace AutomationTool
 			}
 
 			// Read the root BuildGraph element
+			this.SingleNodeName = SingleNodeName;
 			ReadGraphBody(Document.DocumentElement, File.Directory, Arguments);
 			return true;
 		}
@@ -1306,6 +1314,12 @@ namespace AutomationTool
 		/// <param name="ParentNode">The node which owns this task</param>
 		void ReadTask(ScriptElement Element, Node ParentNode)
 		{
+			// If we're running a single node and this element's parent isn't the single node to run, ignore the error and return.
+			if (!string.IsNullOrWhiteSpace(SingleNodeName) && ParentNode.Name != SingleNodeName)
+			{
+				return;
+			}
+
 			if (EvaluateCondition(Element))
 			{
 				// Get the reflection info for this element
