Temporarily remove HoloLens from installed build due to compile errors.

#rnx

diff --git a/Engine/Build/InstalledEngineBuild.xml b/Engine/Build/InstalledEngineBuild.xml
--- a/Engine/Build/InstalledEngineBuild.xml
+++ b/Engine/Build/InstalledEngineBuild.xml
@@ -32,7 +32,7 @@
 	<Property Name="DefaultWithWin32" Value="$(DefaultWithPlatform)"/>
 	<Property Name="DefaultWithWin32" Value="false" If="'$(HostPlatform)' != 'Win64' And !$(AllPlatforms)"/>
 
-	<Property Name="DefaultWithHoloLens" Value="$(DefaultWithPlatform)"/>
+	<Property Name="DefaultWithHoloLens" Value="false"/>
 	<Property Name="DefaultWithHoloLens" Value="false" If="'$(HostPlatform)' != 'Win64' And !$(AllPlatforms)"/>
 
 	<!-- Setup default target platforms for Installed build if they haven't been specified on the commandline -->
