Remove invalid dependency on FindPlatformExtensionSources.sh.

#rnx

diff --git a/Engine/Build/InstalledEngineFilters.xml b/Engine/Build/InstalledEngineFilters.xml
--- a/Engine/Build/InstalledEngineFilters.xml
+++ b/Engine/Build/InstalledEngineFilters.xml
@@ -310,7 +310,6 @@
 		Engine/Build/BatchFiles/Mac/XcodeBuild.sh
 		Engine/Build/BatchFiles/RunUAT.command
 		Engine/Build/BatchFiles/RunUAT.sh
-		Engine/Build/BatchFiles/FindPlatformExtensionSources.sh
 
 		<!-- Mono -->
 		Engine/Binaries/ThirdParty/Mono/Mac/...
@@ -349,7 +348,6 @@
 			Engine/Build/BatchFiles/Linux/RunMono.sh
 			Engine/Build/BatchFiles/Linux/RunXBuild.sh
 			Engine/Build/BatchFiles/RunUAT.sh
-			Engine/Build/BatchFiles/FindPlatformExtensionSources.sh
 
 			<!-- Debug tools -->
 			Engine/Binaries/Linux/dump_syms
