Disable IOS for installed engine builds due to (intermittent?) crashes in IOSToolChain on build farm.

#rnx

diff --git a/Engine/Build/InstalledEngineBuild.xml b/Engine/Build/InstalledEngineBuild.xml
--- a/Engine/Build/InstalledEngineBuild.xml
+++ b/Engine/Build/InstalledEngineBuild.xml
@@ -26,7 +26,7 @@
 	<Property Name="DefaultWithPlatform" Value="true"/>
 	<Property Name="DefaultWithPlatform" Value="false" If="$(HostPlatformEditorOnly) Or $(HostPlatformOnly)"/>
 
-	<Property Name="DefaultWithIOS" Value="$(DefaultWithPlatform)"/>
+	<Property Name="DefaultWithIOS" Value="false"/>
 	<Property Name="DefaultWithIOS" Value="false" If="'$(HostPlatform)' != 'Mac' And !$(AllPlatforms)"/>
 
 	<Property Name="DefaultWithWin32" Value="$(DefaultWithPlatform)"/>
