Fix missing files from installed build on Mac due to preprocessed script using the $(HostPlatform) variable to determine what to include or exclude.

#rnx

diff --git a/Engine/Build/InstalledEngineBuild.xml b/Engine/Build/InstalledEngineBuild.xml
--- a/Engine/Build/InstalledEngineBuild.xml
+++ b/Engine/Build/InstalledEngineBuild.xml
@@ -745,8 +745,8 @@
 			<Delete Files="$(LocalInstalledDir)\..."/>
 
 			<!-- Define filter and exception properties for the installed build -->
-			<Property Name="CopyInstalledFilter" Value="$(CopyEditorFilter)"/>
-			<Property Name="CopyInstalledExceptions" Value="$(CopyEditorExceptions)"/>
+			<Property Name="CopyInstalledFilter" Value="$(CopyEditorFilterWin64)"/>
+			<Property Name="CopyInstalledExceptions" Value="$(CopyEditorExceptionsWin64)"/>
 
 			<!-- Tag files that always need to be part of an installed build -->
 			<Tag Files="#Compile UnrealHeaderTool Win64" With="#Installed Win64"/>
@@ -770,25 +770,25 @@
 			<Do If="'$(WithAndroid)' == true">
 				<Tag Files="#UE4Game Android" Except="#UE4Game Android Unstripped" With="#Installed Win64"/>
 				<Tag Files="#UE4Game Android Stripped" With="#Saved Output"/>
-				<Property Name="CopyInstalledFilter" Value="$(CopyInstalledFilter);$(CopyAndroidFilter)"/>
+				<Property Name="CopyInstalledFilter" Value="$(CopyInstalledFilter);$(CopyAndroidFilterWin64)"/>
 				<Property Name="CopyInstalledExceptions" Value="$(CopyInstalledExceptions);$(CopyAndroidExceptions)"/>
 			</Do>
 			<Do If="'$(WithIOS)' == true">
 				<Tag Files="#UE4Game IOS;#Compile UnrealHeaderTool Mac" Except="#UE4Game IOS Unstripped" With="#Installed Win64"/>
 				<Tag Files="#UE4Game IOS Stripped" With="#Saved Output"/>
-				<Property Name="CopyInstalledFilter" Value="$(CopyInstalledFilter);$(CopyIOSFilter)"/>
+				<Property Name="CopyInstalledFilter" Value="$(CopyInstalledFilter);$(CopyIOSFilterWin64)"/>
 				<Property Name="CopyInstalledExceptions" Value="$(CopyInstalledExceptions);$(CopyIOSExceptions)"/>
 			</Do>
 			<Do If="'$(WithTVOS)' == true">
 				<Tag Files="#UE4Game TVOS" Except="#UE4Game TVOS Unstripped" With="#Installed Win64"/>
 				<Tag Files="#UE4Game TVOS Stripped" With="#Saved Output"/>
-				<Property Name="CopyInstalledFilter" Value="$(CopyInstalledFilter);$(CopyTVOSFilter)"/>
+				<Property Name="CopyInstalledFilter" Value="$(CopyInstalledFilter);$(CopyTVOSFilterWin64)"/>
 				<Property Name="CopyInstalledExceptions" Value="$(CopyInstalledExceptions);$(CopyTVOSExceptions)"/>
 			</Do>
 			<Do If="'$(WithLinux)' == true">
 				<Tag Files="#UE4Game Linux;#Build Tools Linux" Except="#UE4Game Linux Unstripped" With="#Installed Win64"/>
 				<Tag Files="#UE4Game Linux Stripped" With="#Saved Output"/>
-				<Property Name="CopyInstalledFilter" Value="$(CopyInstalledFilter);$(CopyLinuxFilter)"/>
+				<Property Name="CopyInstalledFilter" Value="$(CopyInstalledFilter);$(CopyLinuxFilterWin64)"/>
 				<Property Name="CopyInstalledExceptions" Value="$(CopyInstalledExceptions);$(CopyLinuxExceptions)"/>
 			</Do>
 			<Do If="'$(WithLinuxAArch64)' == true">
@@ -907,8 +907,8 @@
 			<Delete Files="$(LocalInstalledDirMac)/..."/>
 
 			<!-- Define filter and exception properties for the installed build -->
-			<Property Name="CopyInstalledFilter" Value="$(CopyEditorFilter)"/>
-			<Property Name="CopyInstalledExceptions" Value="$(CopyEditorExceptions)"/>
+			<Property Name="CopyInstalledFilter" Value="$(CopyEditorFilterMac)"/>
+			<Property Name="CopyInstalledExceptions" Value="$(CopyEditorExceptionsMac)"/>
 
 			<!-- Tag files that always need to be part of an installed build -->
 			<Tag Files="#Compile UnrealHeaderTool Mac" With="#Installed Mac"/>
@@ -926,7 +926,7 @@
 			<Do  If="'$(WithAndroid)' == true">
 				<Tag Files="#UE4Game Android" Except="#UE4Game Android Unstripped" With="#Installed Mac"/>
 				<Tag Files="#UE4Game Android Stripped" With="#Saved Output"/>
-				<Property Name="CopyInstalledFilter" Value="$(CopyInstalledFilter);$(CopyAndroidFilter)"/>
+				<Property Name="CopyInstalledFilter" Value="$(CopyInstalledFilter);$(CopyAndroidFilterMac)"/>
 				<Property Name="CopyInstalledExceptions" Value="$(CopyInstalledExceptions);$(CopyAndroidExceptions)"/>
 			</Do>
 			<Do If="'$(WithIOS)' == true">
@@ -1024,8 +1024,8 @@
 			<Delete Files="$(LocalInstalledDirLinux)\..."/>
 
 			<!-- Define filter and exception properties for the installed build -->
-			<Property Name="CopyInstalledFilter" Value="$(CopyEditorFilter)"/>
-			<Property Name="CopyInstalledExceptions" Value="$(CopyEditorExceptions)"/>
+			<Property Name="CopyInstalledFilter" Value="$(CopyEditorFilterLinux)"/>
+			<Property Name="CopyInstalledExceptions" Value="$(CopyEditorExceptionsLinux)"/>
 
 			<!-- Tag files that always need to be part of an installed build -->
 			<Tag Files="#Compile UnrealHeaderTool Linux" With="#Installed Linux"/>
diff --git a/Engine/Build/InstalledEngineFilters.xml b/Engine/Build/InstalledEngineFilters.xml
--- a/Engine/Build/InstalledEngineFilters.xml
+++ b/Engine/Build/InstalledEngineFilters.xml
@@ -91,12 +91,6 @@
 		Samples/StarterContent/StarterContent.uproject
 		Samples/MobileStarterContent/MobileStarterContent.uproject
 	</Property>
-	<!-- Don't add VR template for Mac and Linux -->
-	<!--
-	<Property Name="ProjectsToBuildDDC" Value="$(ProjectsToBuildDDC)" If="'$(HostPlatform)' != 'Linux'">
-		Templates/TP_VirtualRealityBP/TP_VirtualRealityBP.uproject
-	</Property>
-	-->
 	
 	<!-- Define Editor Filters -->
 	<Property Name="CopyEditorFilter">
@@ -149,6 +143,7 @@
 		Templates/TP_2DSideScrollerBP/...
 		Templates/TP_VehicleAdv/...
 		Templates/TP_VehicleAdvBP/...
+		Templates/TP_VirtualRealityBP/...
 
 		<!-- Enterprise Templates -->
 		Templates/TP_AEC_BlankBP/...
@@ -210,9 +205,7 @@
 	<!-- Optional API docs -->
 	<Property Name="CopyEditorFilter" Value="$(CopyEditorFilter);Engine/Documentation/Builds/CppAPI-HTML.tgz" If="Exists('Engine/Documentation/Builds/CppAPI-HTML.tgz')"/>
 	<Property Name="CopyEditorFilter" Value="$(CopyEditorFilter);Engine/Documentation/Builds/BlueprintAPI-HTML.tgz" If="Exists('Engine/Documentation/Builds/BlueprintAPI-HTML.tgz')"/>
-	<!-- Don't add VR template for Mac and Linux -->
-	<Property Name="CopyEditorFilter" Value="$(CopyEditorFilter);Templates/TP_VirtualRealityBP/..." If="'$(HostPlatform)' != 'Linux'"/>
-		<!-- Define Editor Exceptions -->
+	<!-- Define Editor Exceptions -->
 	<Property Name="CopyEditorExceptions">
 	
 		<!-- Content -->
@@ -248,53 +241,52 @@
 		Templates/*/manifest.json
 		Templates/*/contents.txt
 	</Property>
-	<!-- exclude VR feature pack on Mac -->
-	<Property Name="CopyEditorExceptions" Value="$(CopyEditorExceptions);FeaturePacks/TP_VirtualRealityBP.upack" If="'$(HostPlatform)' == 'Linux'"/>
 
 	<!-- Projects to be compiled for UAT -->
 	<Property Name="AutomationToolPaths" Value="Engine/Source/Programs/AutomationTool/..."/>
 	
-	<!-- Win64-specific dependencies -->
-	<Do If="'$(HostPlatform)' == 'Win64'">
-		<Property Name="CopyEditorFilter" Value="$(CopyEditorFilter)">
-			<!-- Leap SDK -->
-			Engine/Plugins/Runtime/LeapMotion/ThirdParty/LeapSDK/...
-
-			<!-- Build batch files -->
-			Engine/Build/BatchFiles/Build.bat
-			Engine/Build/BatchFiles/Clean.bat
-			Engine/Build/BatchFiles/Rebuild.bat
-			Engine/Build/BatchFiles/RunUAT.bat
-			Engine/Build/BatchFiles/GetMSBuildPath.bat
-
-			<!-- Build utilities -->
-			Engine/Build/Windows/cl-filter/cl-filter.exe
-			Engine/Build/Windows/link-filter/link-filter.exe
-			Engine/Source/ThirdParty/IntelISPC/bin/Windows/ispc.exe
-
-			<!-- Default resource files -->
-			Engine/Build/Windows/Resources/...
-
-			<!-- Binaries which are explicitly staged by UAT rather than listing runtime dependencies -->
-			Engine/Binaries/ThirdParty/OpenSSL/...
-
-			<!-- Extra tools -->
-			Engine/Extras/3dsMaxScripts/...
-			Engine/Extras/VisualStudioDebugging/UE4.natvis
-			Engine/Extras/MayaLiveLink/...
-			Engine/Extras/MayaVelocityGridExporter/...
-			Engine/Extras/UnrealVS/...
-			Engine/Extras/Redist/en-us/*
+	<!-- Win64 editor specific dependencies -->
+	
+	<Property Name="CopyEditorFilterWin64" Value="$(CopyEditorFilter)">
+		<!-- Leap SDK -->
+		Engine/Plugins/Runtime/LeapMotion/ThirdParty/LeapSDK/...
+
+		<!-- Build batch files -->
+		Engine/Build/BatchFiles/Build.bat
+		Engine/Build/BatchFiles/Clean.bat
+		Engine/Build/BatchFiles/Rebuild.bat
+		Engine/Build/BatchFiles/RunUAT.bat
+		Engine/Build/BatchFiles/GetMSBuildPath.bat
+
+		<!-- Build utilities -->
+		Engine/Build/Windows/cl-filter/cl-filter.exe
+		Engine/Build/Windows/link-filter/link-filter.exe
+		Engine/Source/ThirdParty/IntelISPC/bin/Windows/ispc.exe
+
+		<!-- Default resource files -->
+		Engine/Build/Windows/Resources/...
+
+		<!-- Binaries which are explicitly staged by UAT rather than listing runtime dependencies -->
+		Engine/Binaries/ThirdParty/OpenSSL/...
+
+		<!-- Extra tools -->
+		Engine/Extras/3dsMaxScripts/...
+		Engine/Extras/VisualStudioDebugging/UE4.natvis
+		Engine/Extras/MayaLiveLink/...
+		Engine/Extras/MayaVelocityGridExporter/...
+		Engine/Extras/UnrealVS/...
+		Engine/Extras/Redist/en-us/*
 			
-			<!-- nDisplay Template -->
-			Templates/TP_nDisplayBP/...
-		</Property>
-
-		<Property Name="CopyEditorExceptions" Value="$(CopyEditorExceptions)">
-			<!-- Exclude Mac binaries on windows here, because previous wildcard for plugin resources matches OSX dsym resources on Windows -->
-			Engine/Plugins/.../Binaries/Mac/...
-		</Property>
-	</Do>
+		<!-- nDisplay Template -->
+		Templates/TP_nDisplayBP/...
+	</Property>
+
+	<Property Name="CopyEditorExceptionsWin64" Value="$(CopyEditorExceptions)">
+		<!-- Exclude Mac binaries on windows here, because previous wildcard for plugin resources matches OSX dsym resources on Windows -->
+		Engine/Plugins/.../Binaries/Mac/...
+	</Property>
+
+	<!-- Mac editor specific dependencies -->
 
 	<!-- Files needed to build on Mac, either natively or via the remote toolchain -->
 	<Property Name="MacBuildFiles">
@@ -315,54 +307,54 @@
 		Engine/Binaries/ThirdParty/Mono/Mac/...
 	</Property>
 
-	<!-- Mac-specific dependencies -->
-	<Do If="'$(HostPlatform)' == 'Mac'">
-		<Property Name="CopyEditorFilter" Value="$(CopyEditorFilter);$(MacBuildFiles)">
-			<!-- Extra tools -->
-			Engine/Extras/LLDBDataFormatters/UE4DataFormatters.py
-		</Property>
+	<Property Name="CopyEditorFilterMac" Value="$(CopyEditorFilter);$(MacBuildFiles)">
+		<!-- Extra tools -->
+		Engine/Extras/LLDBDataFormatters/UE4DataFormatters.py
+	</Property>
 
-		<Property Name="CopyEditorExceptions" Value="$(CopyEditorExceptions)">
-			<!-- Exclude Windows binaries on Mac here -->
-			Engine/Plugins/.../Binaries/Win64/...
-			Engine/Plugins/.../Binaries/Win32/...
+	<Property Name="CopyEditorExceptionsMac" Value="$(CopyEditorExceptions)">
+		<!-- Exclude Windows binaries on Mac here -->
+		Engine/Plugins/.../Binaries/Win64/...
+		Engine/Plugins/.../Binaries/Win32/...
 
-			<!-- Don't want these folders, even if they're part of Windows tools -->
-			Engine/Binaries/Win64/...
-			Engine/Binaries/Win32/...
-		</Property>
-	</Do>
+		<!-- Don't want these folders, even if they're part of Windows tools -->
+		Engine/Binaries/Win64/...
+		Engine/Binaries/Win32/...
+	</Property>
 
-	<!-- Linux-specific dependencies -->
-	<Do If="'$(HostPlatform)' == 'Linux'">
-		<Property Name="CopyEditorFilter" Value="$(CopyEditorFilter)">
-			<!-- Build batch files -->
-			Engine/Build/BatchFiles/Linux/Build.sh
-			Engine/Build/BatchFiles/Linux/GenerateGDBInit.sh
-			Engine/Build/BatchFiles/Linux/FixDependencyFiles.sh
-			Engine/Build/BatchFiles/Linux/GenerateProjectFiles.sh
-			Engine/Build/BatchFiles/Linux/QASmokeManual.sh
-			Engine/Build/BatchFiles/Linux/Setup.sh
-			Engine/Build/BatchFiles/Linux/SetupMono.sh
-			Engine/Build/BatchFiles/Linux/FixMonoFiles.sh
-			Engine/Build/BatchFiles/Linux/RunMono.sh
-			Engine/Build/BatchFiles/Linux/RunXBuild.sh
-			Engine/Build/BatchFiles/RunUAT.sh
-
-			<!-- Debug tools -->
-			Engine/Binaries/Linux/dump_syms
-			Engine/Binaries/Linux/BreakpadSymbolEncoder
-
-			<!-- Mono -->
-			Engine/Binaries/ThirdParty/Mono/...
-
-			<!-- Extra tools -->
-			Engine/Extras/GDBPrinters/UE4Printers.py
-
-			<!-- Bundled libcxx -->
-			Engine/Source/ThirdParty/Linux/LibCxx/...
-		</Property>
-	</Do>
+	<!-- Linux editor specific dependencies -->
+	
+	<Property Name="CopyEditorFilterLinux" Value="$(CopyEditorFilter)">
+		<!-- Build batch files -->
+		Engine/Build/BatchFiles/Linux/Build.sh
+		Engine/Build/BatchFiles/Linux/GenerateGDBInit.sh
+		Engine/Build/BatchFiles/Linux/FixDependencyFiles.sh
+		Engine/Build/BatchFiles/Linux/GenerateProjectFiles.sh
+		Engine/Build/BatchFiles/Linux/QASmokeManual.sh
+		Engine/Build/BatchFiles/Linux/Setup.sh
+		Engine/Build/BatchFiles/Linux/SetupMono.sh
+		Engine/Build/BatchFiles/Linux/FixMonoFiles.sh
+		Engine/Build/BatchFiles/Linux/RunMono.sh
+		Engine/Build/BatchFiles/Linux/RunXBuild.sh
+		Engine/Build/BatchFiles/RunUAT.sh
+
+		<!-- Debug tools -->
+		Engine/Binaries/Linux/dump_syms
+		Engine/Binaries/Linux/BreakpadSymbolEncoder
+
+		<!-- Mono -->
+		Engine/Binaries/ThirdParty/Mono/...
+
+		<!-- Extra tools -->
+		Engine/Extras/GDBPrinters/UE4Printers.py
+
+		<!-- Bundled libcxx -->
+		Engine/Source/ThirdParty/Linux/LibCxx/...
+	</Property>
+	<Property Name="CopyEditorExceptionsLinux" Value="$(CopyEditorExceptions)">
+		FeaturePacks/TP_VirtualRealityBP.upack
+		Templates/TP_VirtualRealityBP/...
+	</Property>
 
 	<!-- Target Platform Filters/Exceptions -->
 
@@ -438,26 +430,12 @@
 		Engine/Source/ThirdParty/Vivox/vivox-sdk/VivoxCoreSDK_UPL.xml
 	</Property>
 	
-	<!-- Win64 Specific -->
-	<Do If="'$(HostPlatform)' == 'Win64'">
-		<Property Name="CopyAndroidFilter" Value="$(CopyAndroidFilter)">
-			Engine/Extras/AndroidWorks/Win64/*.exe
-		</Property>
-	</Do>
-	
-	<!-- Mac Specific -->
-	<Do If="'$(HostPlatform)' == 'Mac'">
-		<Property Name="CopyAndroidFilter" Value="$(CopyAndroidFilter)">
-			Engine/Extras/AndroidWorks/Mac/*.dmg
-		</Property>
-	</Do>
-
-	<!-- Linux Specific -->
-	<Do If="'$(HostPlatform)' == 'Linux'">
-		<Property Name="CopyAndroidFilter" Value="$(CopyAndroidFilter)">
-			Engine/Extras/AndroidWorks/Linux/*.run
-		</Property>
-	</Do>
+	<Property Name="CopyAndroidFilterWin64" Value="$(CopyAndroidFilter)">
+		Engine/Extras/AndroidWorks/Win64/*.exe
+	</Property>
+	<Property Name="CopyAndroidFilterMac" Value="$(CopyAndroidFilter)">
+		Engine/Extras/AndroidWorks/Mac/*.dmg
+	</Property>
 
 	<Property Name="CopyAndroidExceptions">
 		Engine/Binaries/Android/....apk
@@ -468,17 +446,15 @@
 		<!-- Build Files -->
 		Engine/Build/IOS/...
 	</Property>
-	<Do If="'$(HostPlatform)' == 'Win64'">
-		<Property Name="CopyIOSFilter" Value="$(CopyIOSFilter)">
-			Engine/Binaries/DotNET/IOS/openssl.exe
-			Engine/Binaries/Mac/DsymExporter*
-			Engine/Binaries/ThirdParty/ICU/icu4c-53_1/Mac/...
-			Engine/Binaries/ThirdParty/IOS/*
-			Engine/Build/BatchFiles/MakeAndInstallSSHKey.bat
-			Engine/Build/Rsync/...
-			Engine/Extras/ThirdPartyNotUE/DeltaCopy/...
-		</Property>
-	</Do>
+	<Property Name="CopyIOSFilterWin64" Value="$(CopyIOSFilter)">
+		Engine/Binaries/DotNET/IOS/openssl.exe
+		Engine/Binaries/Mac/DsymExporter*
+		Engine/Binaries/ThirdParty/ICU/icu4c-53_1/Mac/...
+		Engine/Binaries/ThirdParty/IOS/*
+		Engine/Build/BatchFiles/MakeAndInstallSSHKey.bat
+		Engine/Build/Rsync/...
+		Engine/Extras/ThirdPartyNotUE/DeltaCopy/...
+	</Property>
 
 	<Property Name="CopyIOSExceptions">
 		<!-- Build Files -->
@@ -494,14 +470,12 @@
 	</Property>
 
 	<!-- TVOS on Win64 specific -->
-	<Do If="'$(HostPlatform)' == 'Win64'">
-		<Property Name="CopyTVOSFilter" Value="$(CopyTVOSFilter)">
-			Engine/Binaries/DotNET/IOS/openssl.exe
-			Engine/Binaries/ThirdParty/IOS/*
-			Engine/Build/BatchFiles/MakeAndInstallSSHKey.bat
-			Engine/Extras/ThirdPartyNotUE/DeltaCopy/...
-		</Property>
-	</Do>
+	<Property Name="CopyTVOSFilterWin64" Value="$(CopyTVOSFilter)">
+		Engine/Binaries/DotNET/IOS/openssl.exe
+		Engine/Binaries/ThirdParty/IOS/*
+		Engine/Build/BatchFiles/MakeAndInstallSSHKey.bat
+		Engine/Extras/ThirdPartyNotUE/DeltaCopy/...
+	</Property>
 
 	<Property Name="CopyTVOSExceptions">
 		<!-- Build Files -->
@@ -517,7 +491,7 @@
 		Engine/Source/ThirdParty/Linux/HaveLinuxDependencies
 	</Property>
 
-	<Property Name="CopyLinuxFilter" Value="$(CopyLinuxFilter)" If="'$(HostPlatform)' == 'Win64'">
+	<Property Name="CopyLinuxFilterWin64" Value="$(CopyLinuxFilter)">
 		<!-- Linux cross build Debug tools -->
 		Engine/Binaries/Linux/dump_syms.exe
 		Engine/Binaries/Linux/BreakpadSymbolEncoder.exe
