Add SetupAndroid to binary releases

#ue4
#android

diff --git a/Engine/Build/InstalledEngineFilters.xml b/Engine/Build/InstalledEngineFilters.xml
--- a/Engine/Build/InstalledEngineFilters.xml
+++ b/Engine/Build/InstalledEngineFilters.xml
@@ -407,6 +407,7 @@
 	<Property Name="CopyAndroidFilter">
 		<!-- Build Files -->
 		Engine/Build/Android/...
+		Engine/Extras/Android/...
 		Engine/Source/ThirdParty/Android/...
 		Engine/Source/ThirdParty/AndroidPermission/...
 		Engine/Source/ThirdParty/GoogleARCore/...
@@ -430,13 +431,6 @@
 		Engine/Source/ThirdParty/Vivox/vivox-sdk/VivoxCoreSDK_UPL.xml
 	</Property>
 	
-	<Property Name="CopyAndroidFilterWin64" Value="$(CopyAndroidFilter)">
-		Engine/Extras/AndroidWorks/Win64/*.exe
-	</Property>
-	<Property Name="CopyAndroidFilterMac" Value="$(CopyAndroidFilter)">
-		Engine/Extras/AndroidWorks/Mac/*.dmg
-	</Property>
-
 	<Property Name="CopyAndroidExceptions">
 		Engine/Binaries/Android/....apk
 	</Property>
