Restoring a property that has been missing since Feb. 12th.

#rnx

diff --git a/Engine/Build/InstalledEngineBuild.xml b/Engine/Build/InstalledEngineBuild.xml
--- a/Engine/Build/InstalledEngineBuild.xml
+++ b/Engine/Build/InstalledEngineBuild.xml
@@ -96,6 +96,7 @@
 	<!-- Hololens Architectures -->
 	<Property Name="HoloLensArchitectures" Value="x64+arm64+multi"/>
 
+	<Property Name="AndroidGPUArchitectures" Value="es2"/>
 	<Property Name="LuminGPUArchitectures" Value="es2"/>
 	
 	<!-- Compile flags for Windows Game binaries -->
