Add option for building installed engine with VS2019

diff --git a/Engine/Build/InstalledEngineBuild.xml b/Engine/Build/InstalledEngineBuild.xml
--- a/Engine/Build/InstalledEngineBuild.xml
+++ b/Engine/Build/InstalledEngineBuild.xml
@@ -75,6 +75,9 @@
 
 	<!-- Generate full debug info for binary editor and packaged application builds -->
 	<Option Name="WithFullDebugInfo" Restrict="true|false" DefaultValue="false" Description="Generate full debug info for binary editor and packaged application builds"/>
+
+	<!-- Whether to use VS2019 -->
+	<Option Name="VS2019" Restrict="true|false" DefaultValue="false" Description="Use Visual Studio 2019 to build Windows targets. By default, Visual Studio 2017 is used for maximum compatibility"/>
 	
 	<!-- The local output directory -->
 	<Property Name="LocalInstalledDir" Value="$(RootDir)/LocalBuilds/Engine/Windows"/>
@@ -105,6 +108,7 @@
 
 	<!-- Compile flags for Windows targets -->
 	<Property Name="VSCompilerArg" Value="-2017"/>
+	<Property Name="VSCompilerArg" Value="-2019" If="$(VS2019)"/>
 
 	<!-- Names of the client and server targets to build -->
 	<Property Name="OptionalClientTarget" Value=""/>
