Add DDC for Win64 only

diff --git a/Engine/Build/InstalledEngineBuild.xml b/Engine/Build/InstalledEngineBuild.xml
--- a/Engine/Build/InstalledEngineBuild.xml
+++ b/Engine/Build/InstalledEngineBuild.xml
@@ -633,7 +633,7 @@
 			<Copy From="$(RootDir)" To="$(DDCDir)" Files="#FilteredCopyList"/>
 			<Copy From="$(CsToolsDir)" To="$(DDCDir)" Files="#Build Tools CS Binaries"/>
 
-			<Command Name="BuildDerivedDataCache" Arguments="-TempDir=&quot;$(DDCDir)&quot; -FeaturePacks=&quot;$(ProjectsToBuildDDC)&quot; -TargetPlatforms=$(DDCPlatformsWin64) -HostPlatform=Win64 -SavedDir=&quot;$(SavedDir)&quot;"/>
+			<Command Name="BuildDerivedDataCache" Arguments="-TempDir=&quot;$(DDCDir)&quot; -FeaturePacks=&quot;$(ProjectsToBuildDDCWin64)&quot; -TargetPlatforms=$(DDCPlatformsWin64) -HostPlatform=Win64 -SavedDir=&quot;$(SavedDir)&quot;"/>
 
 			<Tag Files="$(SavedDir)/Engine/DerivedDataCache/Compressed.ddp" With="#CompressedDDCWin64"/>
 		</Node>
diff --git a/Engine/Build/InstalledEngineFilters.xml b/Engine/Build/InstalledEngineFilters.xml
--- a/Engine/Build/InstalledEngineFilters.xml
+++ b/Engine/Build/InstalledEngineFilters.xml
@@ -67,7 +67,6 @@
 		TP_VehicleAdvBP
 		TP_VehicleBP
 		TP_VirtualRealityBP
-		TP_CollaborativeBP
 	</Property>
 
 	<!-- List of projects to build DDC for -->
@@ -84,13 +83,16 @@
 		Templates/TP_PuzzleBP/TP_PuzzleBP.uproject
 		Templates/TP_2DSideScrollerBP/TP_2DSideScrollerBP.uproject
 		Templates/TP_VehicleAdvBP/TP_VehicleAdvBP.uproject
-		Templates/TP_CollaborativeBP/TP_CollaborativeBP.uproject
 		Templates/TP_AEC_BlankBP/TP_AEC_BlankBP.uproject
 		Templates/TP_AEC_ArchvisBP/TP_AEC_ArchvisBP.uproject
 		Templates/TP_PhotoStudioBP/TP_PhotoStudioBP.uproject
 		Samples/StarterContent/StarterContent.uproject
 		Samples/MobileStarterContent/MobileStarterContent.uproject
 	</Property>
+	<!-- Win64 specific DDC -->
+	<Property Name="ProjectsToBuildDDCWin64" Value="$(ProjectsToBuildDDC)">
+		Templates/TP_CollaborativeBP/TP_CollaborativeBP.uproject
+	</Property>
 	
 	<!-- Define Editor Filters -->
 	<Property Name="CopyEditorFilter">
@@ -148,7 +150,6 @@
 		<!-- Enterprise Templates -->
 		Templates/TP_AEC_BlankBP/...
 		Templates/TP_AEC_ArchvisBP/...
-		Templates/TP_CollaborativeBP/...
 		Templates/TP_PhotoStudioBP/...
 
 		Templates/TP_ME_BlankBP/...
@@ -277,8 +278,11 @@
 		Engine/Extras/UnrealVS/...
 		Engine/Extras/Redist/en-us/*
 			
-		<!-- nDisplay Template -->
+		<!-- Virtual Production Templates -->
 		Templates/TP_nDisplayBP/...
+
+		<!-- Enterprise Templates -->
+		Templates/TP_CollaborativeBP/...
 	</Property>
 
 	<Property Name="CopyEditorExceptionsWin64" Value="$(CopyEditorExceptions)">
