Pulling Datasmith plugins to their own build step.

#rnx

diff --git a/Engine/Build/InstalledEngineBuild.xml b/Engine/Build/InstalledEngineBuild.xml
--- a/Engine/Build/InstalledEngineBuild.xml
+++ b/Engine/Build/InstalledEngineBuild.xml
@@ -129,7 +129,7 @@
 		<Node Name="Compile UnrealHeaderTool Win64" Requires="Update Version Files">
 			<Compile Target="UnrealHeaderTool" Platform="Win64" Configuration="Development" Arguments="-precompile -allmodules"/>
 		</Node>
-		<Node Name="Compile UE4Editor Win64" Requires="Compile UnrealHeaderTool Win64" Produces="#UE4Editor Win64;#UE4Editor Win64 Unstripped;#UE4Editor Win64 Stripped;#UE4Editor Win64 Unsigned;#UE4Editor Win64 Signed">
+		<Node Name="Compile UE4Editor Win64" Requires="Compile UnrealHeaderTool Win64" Produces="#UE4Editor Win64;#UE4Editor Win64 Unstripped;#UE4Editor Win64 Stripped;#UE4Editor Win64 Unsigned;#UE4Editor Win64 Signed;#UE4Editor Win64 WithExceptions">
 			<Compile Target="UE4Editor" Platform="Win64" Configuration="DebugGame" Tag="#UE4Editor Win64" Arguments="-precompile -allmodules"/>
 			<Compile Target="UE4Editor" Platform="Win64" Configuration="Development" Tag="#UE4Editor Win64" Arguments="-precompile -allmodules"/>
 
@@ -140,7 +140,7 @@
 				<SrcSrv BinaryFiles="#UE4Editor Win64" SourceFiles="#SourceFiles" Branch="$(Branch)" Change="$(Change)"/>
 			</Do>
 
-			<Tag Files="#UE4Editor Win64" Filter="$(Win64StripFilter)" Except="$(Win64StripExceptions)" With="#UE4Editor Win64 Unstripped"/>
+			<Tag Files="#UE4Editor Win64" Filter="$(Win64StripFilter)" Except="$(Win64StripExceptions);$(PluginsExceptions)" With="#UE4Editor Win64 Unstripped"/>
 
 			<Switch>
 				<Case If="'$(WithFullDebugInfo)' == true">
@@ -154,11 +154,13 @@
 			</Switch>
 
 			<Do If="'$(SignExecutables)' == true">
-				<Tag Files="#UE4Editor Win64" Filter="$(WindowsSignFilter)" With="#UE4Editor Win64 Unsigned"/>
+				<Tag Files="#UE4Editor Win64" Filter="$(WindowsSignFilter)" Except="$(PluginsExceptions)" With="#UE4Editor Win64 Unsigned"/>
 				<Copy Files="#UE4Editor Win64 Unsigned" From="$(RootDir)" To="$(SavedOutput)" Tag="#Editor_ToSign_Win64"/>
 				<Sign Files="#Editor_ToSign_Win64" Tag="#UE4Editor Win64 Signed"/>
 			</Do>
 
+			<Tag Files="#UE4Editor Win64" Except="$(PluginsExceptions)" With="#UE4Editor Win64 WithExceptions"/>
+
 		</Node>
 	</Agent>
 
@@ -723,9 +725,20 @@
 		</Node>
 	</Agent>
 
+	<!-- These are plugins that need to be compiled with a full workspace rather than Win64_Licensee. -->
+	<Agent Name="Extra Plugins Agent" Type="Win64">
+		<Node Name="Compile Datasmith Plugins" Requires="Compile UnrealHeaderTool Win64" Produces="#DatasmithPluginBinaries">
+			<Property Name="DatasmithPlugins" Value="DatasmithCADImporter;DatasmithIFCImporter;DatasmithC4DImporter;AxFImporter;StaticMeshEditorExtension"/>
+			<ForEach Name="Plugin" Values="$(DatasmithPlugins)">
+				<Compile Target="UE4Editor" Configuration="Development" Platform="Win64" Arguments="-Plugin=&quot;$(RootDir)/Engine/Plugins/Enterprise/$(Plugin)/$(Plugin).uplugin&quot;"/>
+				<Tag Files="$(RootDir)/Engine/Plugins/Enterprise/$(Plugin)/Binaries/..." With="#DatasmithPluginBinaries"/>
+			</ForEach>
+		</Node>
+	</Agent>
+
 	<Agent Name="Installed Build Group Win64" Type="Win64_Licensee">
 		<!-- Build up a list of requirements as some can be optional -->
-		<Property Name="InstalledRequirements" Value="Compile UnrealHeaderTool Win64;#UE4Editor Win64;#UE4Editor Win64 Unstripped;#UE4Editor Win64 Stripped;#UE4Editor Win64 Unsigned;#UE4Editor Win64 Signed;#Build Tools Win64;#Build Tools CS Binaries"/>
+		<Property Name="InstalledRequirements" Value="Compile UnrealHeaderTool Win64;#UE4Editor Win64 WithExceptions;#UE4Editor Win64 Unstripped;#UE4Editor Win64 Stripped;#UE4Editor Win64 Unsigned;#UE4Editor Win64 Signed;#Build Tools Win64;#Build Tools CS Binaries"/>
 		<!-- Optional Target Platform requirements -->
 		<Property Name="InstalledRequirements" Value="$(InstalledRequirements);#UE4Game Win64;#UE4Game Win64 Unstripped;#UE4Game Win64 Stripped;#UE4Game Win64 Unsigned;#UE4Game Win64 Signed" If="'$(WithWin64)' == true"/>
 		<Property Name="InstalledRequirements" Value="$(InstalledRequirements);#UE4Game Win32;#UE4Game Win32 Unstripped;#UE4Game Win32 Stripped;#UE4Game Win32 Unsigned;#UE4Game Win32 Signed;#Build Tools Win32" If="'$(WithWin32)' == true"/>
@@ -741,7 +754,7 @@
 		<Property Name="InstalledRequirements" Value="$(InstalledRequirements);#UE4Game HoloLens" If="'$(WithHoloLens)' == true"/>
 		<Property Name="InstalledRequirements" Value="$(InstalledRequirements);#CompressedDDCWin64" If="'$(WithDDC)' == true"/>
 		<!-- Feature packs -->
-		<Property Name="InstalledRequirements" Value="$(InstalledRequirements);Make Feature Packs"/>
+		<Property Name="InstalledRequirements" Value="$(InstalledRequirements);#Compile Datasmith Plugins;#DatasmithPluginBinaries;Make Feature Packs"/>
 
 		<!-- Filters the build products required for an installed build into groups depending on whether they need to be stripped, signed or just copied -->
 		<Node Name="Make Installed Build Win64" Requires="$(InstalledRequirements)" Produces="#Installed Build Win64 Files">
@@ -755,9 +768,10 @@
 
 			<!-- Tag files that always need to be part of an installed build -->
 			<Tag Files="#Compile UnrealHeaderTool Win64" With="#Installed Win64"/>
-			<Tag Files="#UE4Editor Win64" Except="#UE4Editor Win64 Unstripped;#UE4Editor Win64 Unsigned" With="#Installed Win64"/>
+			<Tag Files="#UE4Editor Win64 WithExceptions" Except="#UE4Editor Win64 Unstripped;#UE4Editor Win64 Unsigned" With="#Installed Win64"/>
 			<Tag Files="#UE4Editor Win64 Stripped;#UE4Editor Win64 Signed" With="#Saved Output"/>
 			<Tag Files="#Build Tools Win64;#Make Feature Packs" With="#Installed Win64"/>
+			<Tag Files="#DatasmithPluginBinaries" With="#Installed Win64"/>
 
 			<!-- Tag enabled target platforms -->
 			<Do If="'$(WithWin64)' == true">
diff --git a/Engine/Build/InstalledEngineFilters.xml b/Engine/Build/InstalledEngineFilters.xml
--- a/Engine/Build/InstalledEngineFilters.xml
+++ b/Engine/Build/InstalledEngineFilters.xml
@@ -38,6 +38,14 @@
 
 
 	<Property Name="WindowsSignFilter" Value="*.exe;*.dll"/>
+
+	<Property Name="PluginsExceptions">
+		Engine/Plugins/Enterprise/DatasmithCADImporter/...
+		Engine/Plugins/Enterprise/DatasmithIFCImporter/...
+		Engine/Plugins/Enterprise/DatasmithC4DImporter/...
+		Engine/Plugins/Enterprise/AxFImporter/...
+		Engine/Plugins/Enterprise/StaticMeshEditorExtension/...
+	</Property>
 	<Property Name="MacSignFilter" Value="*.dylib;*.app"/>
 	<Property Name="SignExceptions" Value=".../ThirdParty/..."/>
 
