Fix MDL importer not available in installed builds

#rnx

diff --git a/Engine/Build/InstalledEngineBuild.xml b/Engine/Build/InstalledEngineBuild.xml
--- a/Engine/Build/InstalledEngineBuild.xml
+++ b/Engine/Build/InstalledEngineBuild.xml
@@ -728,7 +728,7 @@
 	<!-- These are plugins that need to be compiled with a full workspace rather than Win64_Licensee. -->
 	<Agent Name="Extra Plugins Agent" Type="Win64">
 		<Node Name="Compile Datasmith Plugins" Requires="Compile UnrealHeaderTool Win64" Produces="#DatasmithPluginBinaries">
-			<Property Name="DatasmithPlugins" Value="DatasmithCADImporter;DatasmithIFCImporter;DatasmithC4DImporter;AxFImporter;StaticMeshEditorExtension"/>
+			<Property Name="DatasmithPlugins" Value="DatasmithCADImporter;DatasmithIFCImporter;DatasmithC4DImporter;AxFImporter;StaticMeshEditorExtension;MDLImporter"/>
 			<ForEach Name="Plugin" Values="$(DatasmithPlugins)">
 				<Compile Target="UE4Editor" Configuration="Development" Platform="Win64" Arguments="-Plugin=&quot;$(RootDir)/Engine/Plugins/Enterprise/$(Plugin)/$(Plugin).uplugin&quot;"/>
 				<Tag Files="$(RootDir)/Engine/Plugins/Enterprise/$(Plugin)/Binaries/..." With="#DatasmithPluginBinaries"/>
diff --git a/Engine/Build/InstalledEngineFilters.xml b/Engine/Build/InstalledEngineFilters.xml
--- a/Engine/Build/InstalledEngineFilters.xml
+++ b/Engine/Build/InstalledEngineFilters.xml
@@ -45,6 +45,7 @@
 		Engine/Plugins/Enterprise/DatasmithC4DImporter/...
 		Engine/Plugins/Enterprise/AxFImporter/...
 		Engine/Plugins/Enterprise/StaticMeshEditorExtension/...
+		Engine/Plugins/Enterprise/MDLImporter/...
 	</Property>
 	<Property Name="MacSignFilter" Value="*.dylib;*.app"/>
 	<Property Name="SignExceptions" Value=".../ThirdParty/..."/>
