Put back CollabViewer template on Mac + Add Hololens template on Win64 only

#rnx

diff --git a/Engine/Build/InstalledEngineFilters.xml b/Engine/Build/InstalledEngineFilters.xml
--- a/Engine/Build/InstalledEngineFilters.xml
+++ b/Engine/Build/InstalledEngineFilters.xml
@@ -76,6 +76,7 @@
 		TP_VehicleAdvBP
 		TP_VehicleBP
 		TP_VirtualRealityBP
+		TP_CollaborativeBP
 	</Property>
 
 	<!-- List of projects to build DDC for -->
@@ -95,12 +96,13 @@
 		Templates/TP_AEC_BlankBP/TP_AEC_BlankBP.uproject
 		Templates/TP_AEC_ArchvisBP/TP_AEC_ArchvisBP.uproject
 		Templates/TP_PhotoStudioBP/TP_PhotoStudioBP.uproject
+		Templates/TP_CollaborativeBP/TP_CollaborativeBP.uproject
 		Samples/StarterContent/StarterContent.uproject
 		Samples/MobileStarterContent/MobileStarterContent.uproject
 	</Property>
 	<!-- Win64 specific DDC -->
 	<Property Name="ProjectsToBuildDDCWin64" Value="$(ProjectsToBuildDDC)">
-		Templates/TP_CollaborativeBP/TP_CollaborativeBP.uproject
+		Templates/TP_HololensBP/TP_HololensBP.uproject
 	</Property>
 	
 	<!-- Define Editor Filters -->
@@ -164,6 +166,8 @@
 		Templates/TP_ME_BlankBP/...
 		Templates/TP_ME_VProdBP/...
 
+		Templates/TP_CollaborativeBP/...
+
 		<!-- Shared template resources -->
 		Templates/TemplateResources/...
 
@@ -288,7 +292,7 @@
 		Templates/TP_nDisplayBP/...
 
 		<!-- Enterprise Templates -->
-		Templates/TP_CollaborativeBP/...
+		Templates/TP_HololensBP/...
 	</Property>
 
 	<Property Name="CopyEditorExceptionsWin64" Value="$(CopyEditorExceptions)">
