Remove Hololens template from DDC build to avoid warning in CIS

#rnx

diff --git a/Engine/Build/InstalledEngineFilters.xml b/Engine/Build/InstalledEngineFilters.xml
--- a/Engine/Build/InstalledEngineFilters.xml
+++ b/Engine/Build/InstalledEngineFilters.xml
@@ -102,7 +102,6 @@
 	</Property>
 	<!-- Win64 specific DDC -->
 	<Property Name="ProjectsToBuildDDCWin64" Value="$(ProjectsToBuildDDC)">
-		Templates/TP_HololensBP/TP_HololensBP.uproject
 	</Property>
 	
 	<!-- Define Editor Filters -->
