Moving DatasmithCADWorker to a NotForLicensee tools build.

#rnx

diff --git a/Engine/Build/InstalledEngineBuild.xml b/Engine/Build/InstalledEngineBuild.xml
--- a/Engine/Build/InstalledEngineBuild.xml
+++ b/Engine/Build/InstalledEngineBuild.xml
@@ -473,11 +473,17 @@
 		</Node>
 	</Agent>
 
+	<Agent Name="NotForLicensee Tools Group Win64" Type="Win64">
+		<Node Name="Build NotForLicensee Tools Win64" Requires="Compile UnrealHeaderTool Win64" Produces="#NotForLicensee Build Tools Win64">
+			<Compile Target="DatasmithCADWorker" Configuration="Development" Platform="Win64" Tag="#NotForLicensee Build Tools Win64"/>
+		</Node>
+	</Agent>
+
 	<Agent Name="Tools Group Win64" Type="Win64_Licensee">
-		<Node Name="Build Tools Win64" Requires="Compile UnrealHeaderTool Win64">
+		<Node Name="Build Tools Win64" Requires="Compile UnrealHeaderTool Win64;#NotForLicensee Build Tools Win64">
 			<Compile Target="CrashReportClient" Configuration="Shipping" Platform="Win64" Tag="#Build Tools Win64" Arguments="$(CrashReporterCompileArgs)"/>
 			<Compile Target="CrashReportClientEditor" Configuration="Shipping" Platform="Win64" Tag="#Build Tools Win64" Arguments="$(CrashReporterCompileArgs)"/>
-			<Compile Target="DatasmithCADWorker" Configuration="Development" Platform="Win64" Tag="#Build Tools Win64"/>
+			<Tag Files="#NotForLicensee Build Tools Win64" With="#Build Tools Win64"/>
 			<Compile Target="UnrealWatchdog" Configuration="Shipping" Platform="Win64" Tag="#Build Tools Win64"/>
 			<Compile Target="ShaderCompileWorker" Configuration="Development" Platform="Win64" Tag="#Build Tools Win64"/>
 			<Compile Target="UnrealCEFSubProcess" Configuration="Development" Platform="Win64" Tag="#Build Tools Win64"/>
