Remove references to AndroidGPUArchitectures. Property was removed in CL 11245496; absence causes warnings when executing the script.

#rnx

diff --git a/Engine/Build/InstalledEngineBuild.xml b/Engine/Build/InstalledEngineBuild.xml
--- a/Engine/Build/InstalledEngineBuild.xml
+++ b/Engine/Build/InstalledEngineBuild.xml
@@ -323,8 +323,8 @@
 
 			<ForEach Name="Target" Values="UE4Game;$(OptionalClientTarget)">
 				<ForEach Name="Configuration" Values="$(GameConfigurations)">
-					<Compile Target="$(Target)" Platform="Android" Configuration="$(Configuration)" Tag="#UE4Game Android" Arguments="-precompile -allmodules -nolink -architectures=$(AndroidArchitectures) -gpuarchitectures=$(AndroidGPUArchitectures) $(TargetDebugInfoArg)"/>
-					<Compile Target="$(Target)" Platform="Android" Configuration="$(Configuration)" Tag="#UE4Game Android" Arguments="-precompile -architectures=$(AndroidArchitectures) -gpuarchitectures=$(AndroidGPUArchitectures) $(TargetDebugInfoArg)" Clean="false"/>
+					<Compile Target="$(Target)" Platform="Android" Configuration="$(Configuration)" Tag="#UE4Game Android" Arguments="-precompile -allmodules -nolink -architectures=$(AndroidArchitectures) $(TargetDebugInfoArg)"/>
+					<Compile Target="$(Target)" Platform="Android" Configuration="$(Configuration)" Tag="#UE4Game Android" Arguments="-precompile -architectures=$(AndroidArchitectures) $(TargetDebugInfoArg)" Clean="false"/>
 				</ForEach>
 			</ForEach>
 
@@ -873,7 +873,7 @@
 
 			<!-- Create command line arguments for any platforms that were built with additional architectures -->
 			<Property Name="ArchitectureCommands" Value=""/>
-			<Property Name="ArchitectureCommands" Value="$(ArchitectureCommands) -AndroidArchitectures=$(AndroidArchitectures) -AndroidGPUArchitectures=$(AndroidGPUArchitectures)" If="'$(WithAndroid)' == true"/>
+			<Property Name="ArchitectureCommands" Value="$(ArchitectureCommands) -AndroidArchitectures=$(AndroidArchitectures)" If="'$(WithAndroid)' == true"/>
 			<Property Name="ArchitectureCommands" Value="$(ArchitectureCommands) -HoloLensArchitectures=$(HoloLensArchitectures)" If="'$(WithHoloLens)' == true"/>
 
 			<!-- Finalize the build and then tag the text file that this process creates -->
@@ -994,7 +994,7 @@
 
 			<!-- Create command line arguments for any platforms that were built with additional architectures -->
 			<Property Name="ArchitectureCommands" Value=""/>
-			<Property Name="ArchitectureCommands" Value="$(ArchitectureCommands) -AndroidArchitectures=$(AndroidArchitectures) -AndroidGPUArchitectures=$(AndroidGPUArchitectures)" If="'$(WithAndroid)' == true"/>
+			<Property Name="ArchitectureCommands" Value="$(ArchitectureCommands) -AndroidArchitectures=$(AndroidArchitectures)" If="'$(WithAndroid)' == true"/>
 			<!-- Finalize the build and then tag the text file that this process creates -->
 			<Command Name="FinalizeInstalledBuild" Arguments="-OutputDir=$(LocalInstalledDirMac) $(ArchitectureCommands) -AnalyticsTypeOverride=$(AnalyticsTypeOverride)"/>
 			<Tag Files="$(LocalInstalledDirMac)/Engine/Build/InstalledBuild.txt" With="#Installed Build Mac Files"/>
@@ -1089,7 +1089,7 @@
 
 			<!-- Create command line arguments for any platforms that were built with additional architectures -->
 			<Property Name="ArchitectureCommands" Value=""/>
-			<Property Name="ArchitectureCommands" Value="$(ArchitectureCommands) -AndroidArchitectures=$(AndroidArchitectures) -AndroidGPUArchitectures=$(AndroidGPUArchitectures)" If="'$(WithAndroid)' == true"/>
+			<Property Name="ArchitectureCommands" Value="$(ArchitectureCommands) -AndroidArchitectures=$(AndroidArchitectures)" If="'$(WithAndroid)' == true"/>
 
 			<!-- Finalize the build and then tag the text file that this process creates -->
 			<Command Name="FinalizeInstalledBuild" Arguments="-OutputDir=$(LocalInstalledDirLinux) -ContentOnlyPlatforms=$(ContentOnlyPlatforms) $(ArchitectureCommands) -AnalyticsTypeOverride=$(AnalyticsTypeOverride)"/>
