Only set CRC API preprocessor defs only be defined if the properties have been set to something.

#rnx

diff --git a/Engine/Build/InstalledEngineBuild.xml b/Engine/Build/InstalledEngineBuild.xml
--- a/Engine/Build/InstalledEngineBuild.xml
+++ b/Engine/Build/InstalledEngineBuild.xml
@@ -119,7 +119,7 @@
 	<!-- Include the script that lists files and patterns used when creating an installed build -->
 	<Include Script="InstalledEngineFilters.xml"/>
 
-	<Property Name="CrashReporterCompileArgs" Value="-CompilerArguments=&quot;/DCRC_TELEMETRY_URL=$(CrashReporterAPIURL) /DCRC_TELEMETRY_KEY=$(CrashReporterAPIKey)&quot;" />
+	<Property Name="CrashReporterCompileArgs" Value="-CompilerArguments=&quot;/DCRC_TELEMETRY_URL=$(CrashReporterAPIURL) /DCRC_TELEMETRY_KEY=$(CrashReporterAPIKey)&quot;" If="'$(CrashReporterAPIURL)' != '' And '$(CrashReporterAPIKey)' != ''" />
 
 	<!-- Base editor and tools for compile on Windows -->
 	<Agent Name="Editor Win64" Type="Win64_Licensee">
