Making Datasmith plugin compile step optional so it can be skipped if building from GitHub.

#rnx

diff --git a/Engine/Build/InstalledEngineBuild.xml b/Engine/Build/InstalledEngineBuild.xml
--- a/Engine/Build/InstalledEngineBuild.xml
+++ b/Engine/Build/InstalledEngineBuild.xml
@@ -26,7 +26,7 @@
 	<Property Name="DefaultWithPlatform" Value="true"/>
 	<Property Name="DefaultWithPlatform" Value="false" If="$(HostPlatformEditorOnly) Or $(HostPlatformOnly)"/>
 
-	<Property Name="DefaultWithIOS" Value="$(DefaultWithPlatform)"/>
+	<Property Name="DefaultWithIOS" Value="true"/>
 	<Property Name="DefaultWithIOS" Value="false" If="'$(HostPlatform)' != 'Mac' And !$(AllPlatforms)"/>
 
 	<Property Name="DefaultWithWin32" Value="$(DefaultWithPlatform)"/>
@@ -734,13 +734,17 @@
 	</Agent>
 
 	<!-- These are plugins that need to be compiled with a full workspace rather than Win64_Licensee. -->
+	<Property Name="CompileDatasmithPlugins" Value="false"/>
+	<Property Name="CompileDatasmithPlugins" Value="true" If="'$(WithDatasmithPlugins)' == 'True'"/>
 	<Agent Name="Extra Plugins Agent" Type="Win64">
 		<Node Name="Compile Datasmith Plugins" Requires="Compile UnrealHeaderTool Win64" Produces="#DatasmithPluginBinaries">
-			<Property Name="DatasmithPlugins" Value="DatasmithCADImporter;DatasmithIFCImporter;DatasmithC4DImporter;AxFImporter;StaticMeshEditorExtension;MDLImporter"/>
-			<ForEach Name="Plugin" Values="$(DatasmithPlugins)">
-				<Compile Target="UE4Editor" Configuration="Development" Platform="Win64" Arguments="-Plugin=&quot;$(RootDir)/Engine/Plugins/Enterprise/$(Plugin)/$(Plugin).uplugin&quot;"/>
-				<Tag Files="$(RootDir)/Engine/Plugins/Enterprise/$(Plugin)/Binaries/..." With="#DatasmithPluginBinaries"/>
-			</ForEach>
+			<Do If="$(CompileDatasmithPlugins) == 'True'">
+				<Property Name="DatasmithPlugins" Value="DatasmithCADImporter;DatasmithIFCImporter;DatasmithC4DImporter;AxFImporter;StaticMeshEditorExtension;MDLImporter"/>
+				<ForEach Name="Plugin" Values="$(DatasmithPlugins)">
+					<Compile Target="UE4Editor" Configuration="Development" Platform="Win64" Arguments="-Plugin=&quot;$(RootDir)/Engine/Plugins/Enterprise/$(Plugin)/$(Plugin).uplugin&quot;"/>
+					<Tag Files="$(RootDir)/Engine/Plugins/Enterprise/$(Plugin)/Binaries/..." With="#DatasmithPluginBinaries"/>
+				</ForEach>
+			</Do>
 		</Node>
 	</Agent>
 
@@ -761,8 +765,9 @@
 		<Property Name="InstalledRequirements" Value="$(InstalledRequirements);#UE4Game Lumin" If="'$(WithLumin)' == true"/>
 		<Property Name="InstalledRequirements" Value="$(InstalledRequirements);#UE4Game HoloLens" If="'$(WithHoloLens)' == true"/>
 		<Property Name="InstalledRequirements" Value="$(InstalledRequirements);#CompressedDDCWin64" If="'$(WithDDC)' == true"/>
+		<Property Name="InstalledRequirements" Value="$(InstalledRequirements);#Compile Datasmith Plugins;#DatasmithPluginBinaries" If="$(CompileDatasmithPlugins)"/>
 		<!-- Feature packs -->
-		<Property Name="InstalledRequirements" Value="$(InstalledRequirements);#Compile Datasmith Plugins;#DatasmithPluginBinaries;Make Feature Packs"/>
+		<Property Name="InstalledRequirements" Value="$(InstalledRequirements);Make Feature Packs"/>
 
 		<!-- Filters the build products required for an installed build into groups depending on whether they need to be stripped, signed or just copied -->
 		<Node Name="Make Installed Build Win64" Requires="$(InstalledRequirements)" Produces="#Installed Build Win64 Files">
@@ -779,7 +784,7 @@
 			<Tag Files="#UE4Editor Win64 WithExceptions" Except="#UE4Editor Win64 Unstripped;#UE4Editor Win64 Unsigned" With="#Installed Win64"/>
 			<Tag Files="#UE4Editor Win64 Stripped;#UE4Editor Win64 Signed" With="#Saved Output"/>
 			<Tag Files="#Build Tools Win64;#Make Feature Packs" With="#Installed Win64"/>
-			<Tag Files="#DatasmithPluginBinaries" With="#Installed Win64"/>
+			<Tag Files="#DatasmithPluginBinaries" With="#Installed Win64" If="$(CompileDatasmithPlugins)"/>
 
 			<!-- Tag enabled target platforms -->
 			<Do If="'$(WithWin64)' == true">
