Implementing feedback from Ben for Datasmith plugin building changes, which also can be used to fix warnings from CrashReporterClient API properties not being set.

#rnx

diff --git a/Engine/Build/InstalledEngineBuild.xml b/Engine/Build/InstalledEngineBuild.xml
--- a/Engine/Build/InstalledEngineBuild.xml
+++ b/Engine/Build/InstalledEngineBuild.xml
@@ -10,6 +10,8 @@
 	<!-- Option to enable all target platforms by default, as we ship in the launcher. Note that this will require cross-compiling IOS on Windows. -->
 	<Option Name="AllPlatforms" Restrict="true|false" DefaultValue="false" Description="Include all target platforms by default"/>
 
+	<Option Name="CompileDatasmithPlugins" Restrict="true|false" DefaultValue="false" Description="If Datasmith plugins should be compiled."/>
+
 	<!-- Get default values for each platform being enabled -->
 	<Property Name="DefaultWithWin64" Value="true"/>
 	<Property Name="DefaultWithWin64" Value="false" If="$(HostPlatformEditorOnly) Or ($(HostPlatformOnly) And '$(HostPlatform)' != 'Win64')"/>
@@ -26,7 +28,7 @@
 	<Property Name="DefaultWithPlatform" Value="true"/>
 	<Property Name="DefaultWithPlatform" Value="false" If="$(HostPlatformEditorOnly) Or $(HostPlatformOnly)"/>
 
-	<Property Name="DefaultWithIOS" Value="true"/>
+	<Property Name="DefaultWithIOS" Value="$(DefaultWithPlatform)"/>
 	<Property Name="DefaultWithIOS" Value="false" If="'$(HostPlatform)' != 'Mac' And !$(AllPlatforms)"/>
 
 	<Property Name="DefaultWithWin32" Value="$(DefaultWithPlatform)"/>
@@ -120,6 +122,8 @@
 	<!-- Include the script that lists files and patterns used when creating an installed build -->
 	<Include Script="InstalledEngineFilters.xml"/>
 
+	<Option Name="CrashReporterAPIURL" DefaultValue="" Description="The URL to use to talk to the CrashReporterClient API."/>
+	<Option Name="CrashReporterAPIKey" DefaultValue="" Description="The API key to use to talk to the CrashReporterClient API."/>
 	<Property Name="CrashReporterCompileArgs" Value="-CompilerArguments=&quot;/DCRC_TELEMETRY_URL=$(CrashReporterAPIURL) /DCRC_TELEMETRY_KEY=$(CrashReporterAPIKey)&quot;" If="'$(CrashReporterAPIURL)' != '' And '$(CrashReporterAPIKey)' != ''" />
 
 	<!-- Base editor and tools for compile on Windows -->
@@ -734,8 +738,6 @@
 	</Agent>
 
 	<!-- These are plugins that need to be compiled with a full workspace rather than Win64_Licensee. -->
-	<Property Name="CompileDatasmithPlugins" Value="false"/>
-	<Property Name="CompileDatasmithPlugins" Value="true" If="'$(WithDatasmithPlugins)' == 'True'"/>
 	<Agent Name="Extra Plugins Agent" Type="Win64">
 		<Node Name="Compile Datasmith Plugins" Requires="Compile UnrealHeaderTool Win64" Produces="#DatasmithPluginBinaries">
 			<Do If="$(CompileDatasmithPlugins) == 'True'">
