Fixed copy of files required for installed builds into platform directories.
#rnx

diff --git a/Engine/Build/InstalledEngineFilters.xml b/Engine/Build/InstalledEngineFilters.xml
--- a/Engine/Build/InstalledEngineFilters.xml
+++ b/Engine/Build/InstalledEngineFilters.xml
@@ -210,7 +210,7 @@
 		Shaders/...
 	</Property>
 	<ForEach Name="Filter" Values="$(CopyEditorEngineOrPlatformFilter)">
-		<Property Name="CopyEditorFilter" Value="$(CopyEditorFilter);Engine/$(Filter)"/>
+		<Property Name="CopyEditorFilter" Value="$(CopyEditorFilter);Engine/$(Filter);Engine/Platforms/*/$(Filter)"/>
 	</ForEach>
 	<!-- Optional target files -->
 	<Property Name="CopyEditorFilter" Value="$(CopyEditorFilter);Engine/Source/UE4Client.Target.cs" If="$(WithClient)"/>
@@ -233,6 +233,7 @@
 		<!-- Plugins -->
 		Engine/Plugins/Runtime/TwitchLiveStreaming/...
 		Engine/Plugins/Runtime/PacketHandlers/CompressionComponents/Oodle/...
+		Engine/Platforms/*/Plugins/Runtime/PacketHandlers/CompressionComponents/Oodle/...
 
 		<!-- Source code -->
 		Engine/Source/Runtime/SQLiteSupport/...
diff --git a/Engine/Source/Programs/UnrealBuildTool/Configuration/UEBuildTarget.cs b/Engine/Source/Programs/UnrealBuildTool/Configuration/UEBuildTarget.cs
--- a/Engine/Source/Programs/UnrealBuildTool/Configuration/UEBuildTarget.cs
+++ b/Engine/Source/Programs/UnrealBuildTool/Configuration/UEBuildTarget.cs
@@ -1242,6 +1242,15 @@ namespace UnrealBuildTool
 				// Add the rules file itself
 				Files.Add(Rules.File);
 
+				// Add the subclass rules
+				if (Rules.SubclassRules != null)
+				{
+					foreach (string SubclassRule in Rules.SubclassRules)
+					{
+						Files.Add(new FileReference(SubclassRule));
+					}
+				}
+
 				// Get a list of all the library paths
 				List<string> LibraryPaths = new List<string>();
 				LibraryPaths.Add(Directory.GetCurrentDirectory());
