Fix for Windows binary build , Android looks for "-armv764" instead of "armv64"(like Hololens does)

diff --git a/Engine/Build/InstalledEngineBuild.xml b/Engine/Build/InstalledEngineBuild.xml
--- a/Engine/Build/InstalledEngineBuild.xml
+++ b/Engine/Build/InstalledEngineBuild.xml
@@ -96,7 +96,8 @@
 	<Property Name="ToolchainDir" Value="$(RootDir)/Engine/Extras/ThirdPartyNotUE/SDKs"/>
 
 	<!-- Architectures that we build for Android. NOTE: This should match the architectures specified in BuildPluginCommand. -->
-	<Property Name="AndroidArchitectures" Value="armv7+arm64"/>
+	<Property Name="AndroidArchitectures" Value="-armv7+-arm64"/>
+	<Property Name="LuminArchitectures" Value="armv7+arm64"/>
 	
 	<!-- Hololens Architectures -->
 	<Property Name="HoloLensArchitectures" Value="x64+arm64+multi"/>
@@ -398,8 +399,8 @@
 
 			<ForEach Name="Target" Values="UE4Game;$(OptionalClientTarget)">
 				<ForEach Name="Configuration" Values="$(GameConfigurations)">
-					<Compile Target="$(Target)" Platform="Lumin" Configuration="$(Configuration)" Tag="#UE4Game Lumin" Arguments="-precompile -allmodules -nolink -architectures=$(AndroidArchitectures) -gpuarchitectures=$(LuminGPUArchitectures) $(TargetDebugInfoArg)"/>
-					<Compile Target="$(Target)" Platform="Lumin" Configuration="$(Configuration)" Tag="#UE4Game Lumin" Arguments="-precompile -architectures=$(AndroidArchitectures) -gpuarchitectures=$(LuminGPUArchitectures) $(TargetDebugInfoArg)" Clean="false"/>
+					<Compile Target="$(Target)" Platform="Lumin" Configuration="$(Configuration)" Tag="#UE4Game Lumin" Arguments="-precompile -allmodules -nolink -architectures=$(LuminArchitectures) -gpuarchitectures=$(LuminGPUArchitectures) $(TargetDebugInfoArg)"/>
+					<Compile Target="$(Target)" Platform="Lumin" Configuration="$(Configuration)" Tag="#UE4Game Lumin" Arguments="-precompile -architectures=$(LuminArchitectures) -gpuarchitectures=$(LuminGPUArchitectures) $(TargetDebugInfoArg)" Clean="false"/>
 				</ForEach>
 			</ForEach>
 			
@@ -466,8 +467,8 @@
 
 			<ForEach Name="Target" Values="UE4Game;$(OptionalClientTarget)">
 				<ForEach Name="Configuration" Values="$(GameConfigurations)">
-					<Compile Target="$(Target)" Platform="Lumin" Configuration="$(Configuration)" Tag="#UE4Game LuminMac" Arguments="-precompile -allmodules -nolink -architectures=$(AndroidArchitectures) -gpuarchitectures=$(LuminGPUArchitectures) $(TargetDebugInfoArg)" AllowParallelExecutor="false"/>
-					<Compile Target="$(Target)" Platform="Lumin" Configuration="$(Configuration)" Tag="#UE4Game LuminMac" Arguments="-precompile -architectures=$(AndroidArchitectures) -gpuarchitectures=$(LuminGPUArchitectures) $(TargetDebugInfoArg)" AllowParallelExecutor="false" Clean="false"/>
+					<Compile Target="$(Target)" Platform="Lumin" Configuration="$(Configuration)" Tag="#UE4Game LuminMac" Arguments="-precompile -allmodules -nolink -architectures=$(LuminArchitectures) -gpuarchitectures=$(LuminGPUArchitectures) $(TargetDebugInfoArg)" AllowParallelExecutor="false"/>
+					<Compile Target="$(Target)" Platform="Lumin" Configuration="$(Configuration)" Tag="#UE4Game LuminMac" Arguments="-precompile -architectures=$(LuminArchitectures) -gpuarchitectures=$(LuminGPUArchitectures) $(TargetDebugInfoArg)" AllowParallelExecutor="false" Clean="false"/>
 				</ForEach>
 			</ForEach>
 
