Fix missing platforms in installed build, if the host machine does not have the necessary SDKs installed.

#rnx

diff --git a/Engine/Build/InstalledEngineBuild.xml b/Engine/Build/InstalledEngineBuild.xml
--- a/Engine/Build/InstalledEngineBuild.xml
+++ b/Engine/Build/InstalledEngineBuild.xml
@@ -777,6 +777,9 @@
 			<!-- Clear any existing files from the final location -->
 			<Delete Files="$(LocalInstalledDir)\..."/>
 
+			<!-- List of platform names -->
+			<Property Name="Platforms" Value=""/>
+			
 			<!-- Define filter and exception properties for the installed build -->
 			<Property Name="CopyInstalledFilter" Value="$(CopyEditorFilterWin64)"/>
 			<Property Name="CopyInstalledExceptions" Value="$(CopyEditorExceptionsWin64)"/>
@@ -792,67 +795,79 @@
 			<Do If="'$(WithWin64)' == true">
 				<Tag Files="#UE4Game Win64" Except="#UE4Game Win64 Unstripped;#UE4Game Win64 Unsigned" With="#Installed Win64"/>
 				<Tag Files="#UE4Game Win64 Stripped;#UE4Game Win64 Signed" With="#Saved Output"/>
+				<Property Name="Platforms" Value="$(Platforms)Win64;"/>
 				<Property Name="CopyInstalledFilter" Value="$(CopyInstalledFilter);$(CopyWin64Filter)"/>
 				<Property Name="CopyInstalledExceptions" Value="$(CopyInstalledExceptions);$(CopyWin64Exceptions)"/>
 			</Do>
 			<Do If="'$(WithWin32)' == true">
 				<Tag Files="#UE4Game Win32;#Build Tools Win32" Except="#UE4Game Win32 Unstripped;#UE4Game Win32 Unsigned" With="#Installed Win64"/>
 				<Tag Files="#UE4Game Win32 Stripped;#UE4Game Win32 Signed" With="#Saved Output"/>
+				<Property Name="Platforms" Value="$(Platforms)Win32;"/>
 				<Property Name="CopyInstalledFilter" Value="$(CopyInstalledFilter);$(CopyWin32Filter)"/>
 				<Property Name="CopyInstalledExceptions" Value="$(CopyInstalledExceptions);$(CopyWin32Exceptions)"/>
 			</Do>
 			<Do If="'$(WithAndroid)' == true">
 				<Tag Files="#UE4Game Android" Except="#UE4Game Android Unstripped" With="#Installed Win64"/>
 				<Tag Files="#UE4Game Android Stripped" With="#Saved Output"/>
+				<Property Name="Platforms" Value="$(Platforms)Android;"/>
 				<Property Name="CopyInstalledFilter" Value="$(CopyInstalledFilter);$(CopyAndroidFilter)"/>
 				<Property Name="CopyInstalledExceptions" Value="$(CopyInstalledExceptions);$(CopyAndroidExceptions)"/>
 			</Do>
 			<Do If="'$(WithIOS)' == true">
 				<Tag Files="#UE4Game IOS;#Compile UnrealHeaderTool Mac" Except="#UE4Game IOS Unstripped" With="#Installed Win64"/>
 				<Tag Files="#UE4Game IOS Stripped" With="#Saved Output"/>
+				<Property Name="Platforms" Value="$(Platforms)IOS;"/>
 				<Property Name="CopyInstalledFilter" Value="$(CopyInstalledFilter);$(CopyIOSFilterWin64)"/>
 				<Property Name="CopyInstalledExceptions" Value="$(CopyInstalledExceptions);$(CopyIOSExceptions)"/>
 			</Do>
 			<Do If="'$(WithTVOS)' == true">
 				<Tag Files="#UE4Game TVOS" Except="#UE4Game TVOS Unstripped" With="#Installed Win64"/>
 				<Tag Files="#UE4Game TVOS Stripped" With="#Saved Output"/>
+				<Property Name="Platforms" Value="$(Platforms)TVOS;"/>
 				<Property Name="CopyInstalledFilter" Value="$(CopyInstalledFilter);$(CopyTVOSFilterWin64)"/>
 				<Property Name="CopyInstalledExceptions" Value="$(CopyInstalledExceptions);$(CopyTVOSExceptions)"/>
 			</Do>
 			<Do If="'$(WithLinux)' == true">
 				<Tag Files="#UE4Game Linux;#Build Tools Linux" Except="#UE4Game Linux Unstripped" With="#Installed Win64"/>
 				<Tag Files="#UE4Game Linux Stripped" With="#Saved Output"/>
+				<Property Name="Platforms" Value="$(Platforms)Linux;"/>
 				<Property Name="CopyInstalledFilter" Value="$(CopyInstalledFilter);$(CopyLinuxFilterWin64)"/>
 				<Property Name="CopyInstalledExceptions" Value="$(CopyInstalledExceptions);$(CopyLinuxExceptions)"/>
 			</Do>
 			<Do If="'$(WithLinuxAArch64)' == true">
 				<Tag Files="#UE4Game LinuxAArch64" Except="#UE4Game LinuxAArch64 Unstripped" With="#Installed Win64"/>
 				<Tag Files="#UE4Game LinuxAArch64 Stripped" With="#Saved Output"/>
+				<Property Name="Platforms" Value="$(Platforms)LinuxAArch64;"/>
 				<Property Name="CopyInstalledFilter" Value="$(CopyInstalledFilter);$(CopyLinuxFilter)"/>
 				<Property Name="CopyInstalledExceptions" Value="$(CopyInstalledExceptions);$(CopyLinuxExceptions)"/>
 			</Do>
 			<Do If="'$(WithPS4)' == true">
 				<Tag Files="#UE4Game PS4" With="#Installed Win64"/>
+				<Property Name="Platforms" Value="$(Platforms)PS4;"/>
 				<Property Name="CopyInstalledFilter" Value="$(CopyInstalledFilter);$(CopyPS4Filter)"/>
 				<Property Name="CopyInstalledExceptions" Value="$(CopyInstalledExceptions);$(CopyPS4Exceptions)"/>
 			</Do>
 			<Do If="'$(WithXboxOne)' == true">
 				<Tag Files="#UE4Game XboxOne" With="#Installed Win64"/>
+				<Property Name="Platforms" Value="$(Platforms)XboxOne;"/>
 				<Property Name="CopyInstalledFilter" Value="$(CopyInstalledFilter);$(CopyXboxOneFilter)"/>
 				<Property Name="CopyInstalledExceptions" Value="$(CopyInstalledExceptions);$(CopyXboxOneExceptions)"/>
 			</Do>
 			<Do If="'$(WithSwitch)' == true">
 				<Tag Files="#UE4Game Switch" With="#Installed Win64"/>
+				<Property Name="Platforms" Value="$(Platforms)Switch;"/>
 				<Property Name="CopyInstalledFilter" Value="$(CopyInstalledFilter);$(CopySwitchFilter)"/>
 				<Property Name="CopyInstalledExceptions" Value="$(CopyInstalledExceptions);$(CopySwitchExceptions)"/>
 			</Do>
 			<Do If="'$(WithLumin)' == true">
 				<Tag Files="#UE4Game Lumin" With="#Installed Win64"/>
+				<Property Name="Platforms" Value="$(Platforms)Lumin;"/>
 				<Property Name="CopyInstalledFilter" Value="$(CopyInstalledFilter);$(CopyLuminFilter)"/>
 				<Property Name="CopyInstalledExceptions" Value="$(CopyInstalledExceptions);$(CopyLuminExceptions)"/>
 			</Do>
 			<Do If="'$(WithHoloLens)' == true">
 				<Tag Files="#UE4Game HoloLens;#UE4Game HoloLens Includes" With="#Installed Win64"/>
+				<Property Name="Platforms" Value="$(Platforms)HoloLens;"/>
 				<Property Name="CopyInstalledFilter" Value="$(CopyInstalledFilter);$(CopyHoloLensFilter)"/>
 				<Property Name="CopyInstalledExceptions" Value="$(CopyInstalledExceptions);$(CopyHoloLensExceptions)"/>
 			</Do>
@@ -911,7 +926,7 @@
 			<Property Name="ArchitectureCommands" Value="$(ArchitectureCommands) -HoloLensArchitectures=$(HoloLensArchitectures)" If="'$(WithHoloLens)' == true"/>
 
 			<!-- Finalize the build and then tag the text file that this process creates -->
-			<Command Name="FinalizeInstalledBuild" Arguments="-OutputDir=&quot;$(LocalInstalledDir)&quot; -ContentOnlyPlatforms=$(ContentOnlyPlatforms) $(ArchitectureCommands) -AnalyticsTypeOverride=$(AnalyticsTypeOverride)"/>
+			<Command Name="FinalizeInstalledBuild" Arguments="-OutputDir=&quot;$(LocalInstalledDir)&quot; -Platforms=$(Platforms) -ContentOnlyPlatforms=$(ContentOnlyPlatforms) $(ArchitectureCommands) -AnalyticsTypeOverride=$(AnalyticsTypeOverride)"/>
 			<Tag Files="$(LocalInstalledDir)/Engine/Build/InstalledBuild.txt" With="#Installed Build Win64 Files"/>
 			
 			<!-- Sanitize all the receipts in the target directory -->
@@ -940,6 +955,9 @@
 			<!-- Clear any existing files from the final location -->
 			<Delete Files="$(LocalInstalledDirMac)/..."/>
 
+			<!-- List of platform names -->
+			<Property Name="Platforms" Value=""/>
+			
 			<!-- Define filter and exception properties for the installed build -->
 			<Property Name="CopyInstalledFilter" Value="$(CopyEditorFilterMac)"/>
 			<Property Name="CopyInstalledExceptions" Value="$(CopyEditorExceptionsMac)"/>
@@ -954,30 +972,35 @@
 			<Do If="'$(WithMac)' == true">
 				<Tag Files="#UE4Game Mac" Except="#UE4Game Mac Unstripped;#UE4Game Mac Unsigned" With="#Installed Mac"/>
 				<Tag Files="#UE4Game Mac Stripped;#UE4Game Mac Signed" With="#Saved Output"/>
+				<Property Name="Platforms" Value="$(Platforms)Mac;"/>
 				<Property Name="CopyInstalledFilter" Value="$(CopyInstalledFilter);$(CopyMacFilter)"/>
 				<Property Name="CopyInstalledExceptions" Value="$(CopyInstalledExceptions);$(CopyMacExceptions)"/>
 			</Do>
 			<Do  If="'$(WithAndroid)' == true">
 				<Tag Files="#UE4Game Android" Except="#UE4Game Android Unstripped" With="#Installed Mac"/>
 				<Tag Files="#UE4Game Android Stripped" With="#Saved Output"/>
+				<Property Name="Platforms" Value="$(Platforms)Android;"/>
 				<Property Name="CopyInstalledFilter" Value="$(CopyInstalledFilter);$(CopyAndroidFilter)"/>
 				<Property Name="CopyInstalledExceptions" Value="$(CopyInstalledExceptions);$(CopyAndroidExceptions)"/>
 			</Do>
 			<Do If="'$(WithIOS)' == true">
 				<Tag Files="#UE4Game IOS" Except="#UE4Game IOS Unstripped" With="#Installed Mac"/>
 				<Tag Files="#UE4Game IOS Stripped" With="#Saved Output"/>
+				<Property Name="Platforms" Value="$(Platforms)IOS;"/>
 				<Property Name="CopyInstalledFilter" Value="$(CopyInstalledFilter);$(CopyIOSFilter)"/>
 				<Property Name="CopyInstalledExceptions" Value="$(CopyInstalledExceptions);$(CopyIOSExceptions)"/>
 			</Do>
 			<Do If="'$(WithTVOS)' == true">
 				<Tag Files="#UE4Game TVOS" Except="#UE4Game TVOS Unstripped" With="#Installed Mac"/>
 				<Tag Files="#UE4Game TVOS Stripped" With="#Saved Output"/>
+				<Property Name="Platforms" Value="$(Platforms)TVOS;"/>
 				<Property Name="CopyInstalledFilter" Value="$(CopyInstalledFilter);$(CopyTVOSFilter)"/>
 				<Property Name="CopyInstalledExceptions" Value="$(CopyInstalledExceptions);$(CopyTVOSExceptions)"/>
 			</Do>
 			<Do If="'$(WithLuminMac)' == true">
 				<Tag Files="#UE4Game LuminMac" Except="#UE4Game LuminMac Unstripped" With="#Installed Mac"/>
 				<Tag Files="#UE4Game LuminMac Stripped" With="#Saved Output"/>
+				<Property Name="Platforms" Value="$(Platforms)Lumin;"/>
 				<Property Name="CopyInstalledFilter" Value="$(CopyInstalledFilter);$(CopyLuminFilter)"/>
 				<Property Name="CopyInstalledExceptions" Value="$(CopyInstalledExceptions);$(CopyLuminExceptions)"/>
 			</Do>
@@ -1030,7 +1053,7 @@
 			<Property Name="ArchitectureCommands" Value=""/>
 			<Property Name="ArchitectureCommands" Value="$(ArchitectureCommands) -AndroidArchitectures=$(AndroidArchitectures) -AndroidGPUArchitectures=$(AndroidGPUArchitectures)" If="'$(WithAndroid)' == true"/>
 			<!-- Finalize the build and then tag the text file that this process creates -->
-			<Command Name="FinalizeInstalledBuild" Arguments="-OutputDir=$(LocalInstalledDirMac) $(ArchitectureCommands) -AnalyticsTypeOverride=$(AnalyticsTypeOverride)"/>
+			<Command Name="FinalizeInstalledBuild" Arguments="-OutputDir=$(LocalInstalledDirMac) -Platforms=$(Platforms) $(ArchitectureCommands) -AnalyticsTypeOverride=$(AnalyticsTypeOverride)"/>
 			<Tag Files="$(LocalInstalledDirMac)/Engine/Build/InstalledBuild.txt" With="#Installed Build Mac Files"/>
 
 			<!-- Sanitize all the receipts in the target directory -->
@@ -1057,6 +1080,9 @@
 			<!-- Clear any existing files from the final location -->
 			<Delete Files="$(LocalInstalledDirLinux)\..."/>
 
+			<!-- List of platform names -->
+			<Property Name="Platforms" Value=""/>
+
 			<!-- Define filter and exception properties for the installed build -->
 			<Property Name="CopyInstalledFilter" Value="$(CopyEditorFilterLinux)"/>
 			<Property Name="CopyInstalledExceptions" Value="$(CopyEditorExceptionsLinux)"/>
@@ -1071,12 +1097,14 @@
 			<Do If="'$(WithLinux)' == true">
 				<Tag Files="#UE4Game Linux;#Build Tools Linux" Except="#UE4Game Linux Unstripped" With="#Installed Linux"/>
 				<Tag Files="#UE4Game Linux Stripped" With="#Saved Output"/>
+				<Property Name="Platforms" Value="$(Platforms)Linux;"/>
 				<Property Name="CopyInstalledFilter" Value="$(CopyInstalledFilter);$(CopyLinuxFilter)"/>
 				<Property Name="CopyInstalledExceptions" Value="$(CopyInstalledExceptions);$(CopyLinuxExceptions)"/>
 			</Do>
 			<Do If="'$(WithLinuxAArch64)' == true">
 				<Tag Files="#UE4Game LinuxAArch64" Except="#UE4Game LinuxAArch64 Unstripped" With="#Installed Linux"/>
 				<Tag Files="#UE4Game LinuxAArch64 Stripped" With="#Saved Output"/>
+				<Property Name="Platforms" Value="$(Platforms)LinuxAArch64;"/>
 				<Property Name="CopyInstalledFilter" Value="$(CopyInstalledFilter);$(CopyLinuxFilter)"/>
 				<Property Name="CopyInstalledExceptions" Value="$(CopyInstalledExceptions);$(CopyLinuxExceptions)"/>
 			</Do>
@@ -1126,7 +1154,7 @@
 			<Property Name="ArchitectureCommands" Value="$(ArchitectureCommands) -AndroidArchitectures=$(AndroidArchitectures) -AndroidGPUArchitectures=$(AndroidGPUArchitectures)" If="'$(WithAndroid)' == true"/>
 
 			<!-- Finalize the build and then tag the text file that this process creates -->
-			<Command Name="FinalizeInstalledBuild" Arguments="-OutputDir=$(LocalInstalledDirLinux) -ContentOnlyPlatforms=$(ContentOnlyPlatforms) $(ArchitectureCommands) -AnalyticsTypeOverride=$(AnalyticsTypeOverride)"/>
+			<Command Name="FinalizeInstalledBuild" Arguments="-OutputDir=$(LocalInstalledDirLinux) -Platforms=$(Platforms) -ContentOnlyPlatforms=$(ContentOnlyPlatforms) $(ArchitectureCommands) -AnalyticsTypeOverride=$(AnalyticsTypeOverride)"/>
 			<Tag Files="$(LocalInstalledDirLinux)/Engine/Build/InstalledBuild.txt" With="#Installed Build Linux Files"/>
 
 			<!-- Sanitize all the receipts in the target directory -->
diff --git a/Engine/Source/Programs/AutomationTool/Scripts/FinalizeInstalledBuild.cs b/Engine/Source/Programs/AutomationTool/Scripts/FinalizeInstalledBuild.cs
--- a/Engine/Source/Programs/AutomationTool/Scripts/FinalizeInstalledBuild.cs
+++ b/Engine/Source/Programs/AutomationTool/Scripts/FinalizeInstalledBuild.cs
@@ -16,6 +16,7 @@ namespace AutomationTool
 	[Help("Command to perform additional steps to prepare an installed build.")]
 	[Help("OutputDir=<RootDirectory>", "Root Directory of the installed build data (required)")]
 	[Help("ContentOnlyPlatforms=<PlatformNameList>", "List of platforms that should be marked as only supporting content projects (optional)")]
+	[Help("Platforms=<PlatformList>)", "List of platforms to add")]
 	[Help("<Platform>Architectures=<ArchitectureList>)", "List of architectures that are used for a given platform (optional)")]
 	[Help("<Platform>GPUArchitectures=<GPUArchitectureList>)", "List of GPU architectures that are used for a given platform (optional)")]
 	[Help("AnalyticsTypeOverride=<TypeName>", "Name to give this build for analytics purposes (optional)")]
@@ -26,13 +27,9 @@ namespace AutomationTool
 		/// </summary>
 		public override void ExecuteBuild()
 		{
-			string OutputDir = ParseParamValue("OutputDir");
-			string ContentOnlyPlatformsString = ParseParamValue("ContentOnlyPlatforms");
-			IEnumerable<UnrealTargetPlatform> ContentOnlyPlatforms = Enumerable.Empty<UnrealTargetPlatform>();
-			if (!String.IsNullOrWhiteSpace(ContentOnlyPlatformsString))
-			{
-				ContentOnlyPlatforms = ContentOnlyPlatformsString.Split(';').Where(x => !String.IsNullOrWhiteSpace(x)).Select(x => UnrealTargetPlatform.Parse(x));
-			}
+			string OutputDir = ParseRequiredStringParam("OutputDir");
+			List<UnrealTargetPlatform> Platforms = ParsePlatformsParamValue("Platforms");
+			List<UnrealTargetPlatform> ContentOnlyPlatforms = ParsePlatformsParamValue("ContentOnlyPlatforms");
 			string AnalyticsTypeOverride = ParseParamValue("AnalyticsTypeOverride");
 
 			// Write InstalledBuild.txt to indicate Engine is installed
@@ -75,90 +72,87 @@ namespace AutomationTool
 			InstalledConfigs.Add(new InstalledPlatformInfo.InstalledPlatformConfiguration(UnrealTargetConfiguration.Development, HostPlatform.Current.HostEditorPlatform, TargetRules.TargetType.Editor, EditorArchitecture, "", EProjectType.Unknown, false));
 			InstalledConfigs.Add(new InstalledPlatformInfo.InstalledPlatformConfiguration(UnrealTargetConfiguration.DebugGame, HostPlatform.Current.HostEditorPlatform, TargetRules.TargetType.Editor, EditorArchitecture, "", EProjectType.Unknown, false));
 
-			foreach (UnrealTargetPlatform CodeTargetPlatform in UnrealTargetPlatform.GetValidPlatforms())
+			foreach (UnrealTargetPlatform CodeTargetPlatform in Platforms)
 			{
-				if (PlatformExports.IsPlatformAvailable(CodeTargetPlatform))
-				{
-					string Architecture = PlatformExports.GetDefaultArchitecture(CodeTargetPlatform, null);
+				string Architecture = PlatformExports.GetDefaultArchitecture(CodeTargetPlatform, null);
 
-					// Try to parse additional Architectures from the command line
-					string Architectures = ParseParamValue(CodeTargetPlatform.ToString() + "Architectures");
-					string GPUArchitectures = ParseParamValue(CodeTargetPlatform.ToString() + "GPUArchitectures");
+				// Try to parse additional Architectures from the command line
+				string Architectures = ParseParamValue(CodeTargetPlatform.ToString() + "Architectures");
+				string GPUArchitectures = ParseParamValue(CodeTargetPlatform.ToString() + "GPUArchitectures");
 
-					// Build a list of pre-compiled architecture combinations for this platform if any
-					List<string> AllArchNames;
+				// Build a list of pre-compiled architecture combinations for this platform if any
+				List<string> AllArchNames;
 
-					if (!String.IsNullOrWhiteSpace(Architectures) && !String.IsNullOrWhiteSpace(GPUArchitectures))
-					{
-						AllArchNames = (from Arch in Architectures.Split('+')
-										from GPUArch in GPUArchitectures.Split('+')
-										select "-" + Arch + "-" + GPUArch).ToList();
-					}
-					else if (!String.IsNullOrWhiteSpace(Architectures))
-					{
-						AllArchNames = Architectures.Split('+').ToList();
-					}
-					// if there aren't any, use the default
-					else
-					{
-						AllArchNames = new List<string>() { Architecture };
-					}
+				if (!String.IsNullOrWhiteSpace(Architectures) && !String.IsNullOrWhiteSpace(GPUArchitectures))
+				{
+					AllArchNames = (from Arch in Architectures.Split('+')
+									from GPUArch in GPUArchitectures.Split('+')
+									select "-" + Arch + "-" + GPUArch).ToList();
+				}
+				else if (!String.IsNullOrWhiteSpace(Architectures))
+				{
+					AllArchNames = Architectures.Split('+').ToList();
+				}
+				// if there aren't any, use the default
+				else
+				{
+					AllArchNames = new List<string>() { Architecture };
+				}
 
-					// Check whether this platform should only be used for content based projects
-					EProjectType ProjectType = ContentOnlyPlatforms.Contains(CodeTargetPlatform) ? EProjectType.Content : EProjectType.Any;
+				// Check whether this platform should only be used for content based projects
+				EProjectType ProjectType = ContentOnlyPlatforms.Contains(CodeTargetPlatform) ? EProjectType.Content : EProjectType.Any;
 
-					// Allow Content only platforms to be shown as options in all projects
-					bool bCanBeDisplayed = ProjectType == EProjectType.Content;
-					foreach (UnrealTargetConfiguration CodeTargetConfiguration in Enum.GetValues(typeof(UnrealTargetConfiguration)))
+				// Allow Content only platforms to be shown as options in all projects
+				bool bCanBeDisplayed = ProjectType == EProjectType.Content;
+				foreach (UnrealTargetConfiguration CodeTargetConfiguration in Enum.GetValues(typeof(UnrealTargetConfiguration)))
+				{
+					Dictionary<String, TargetType> Targets = new Dictionary<string, TargetType>() {
+						{ "UE4Game", TargetType.Game },
+						{ "UE4Client", TargetType.Client },
+						{ "UE4Server", TargetType.Server }
+					};
+					foreach (KeyValuePair<string, TargetType> Target in Targets)
 					{
-						Dictionary<String, TargetType> Targets = new Dictionary<string, TargetType>() {
-							{ "UE4Game", TargetType.Game },
-							{ "UE4Client", TargetType.Client },
-							{ "UE4Server", TargetType.Server }
-						};
-						foreach (KeyValuePair<string, TargetType> Target in Targets)
-						{
-							string CurrentTargetName = Target.Key;
-							TargetType CurrentTargetType = Target.Value;
+						string CurrentTargetName = Target.Key;
+						TargetType CurrentTargetType = Target.Value;
 
-							// Need to check for development receipt as we use that for the Engine code in DebugGame
-							UnrealTargetConfiguration EngineConfiguration = (CodeTargetConfiguration == UnrealTargetConfiguration.DebugGame) ? UnrealTargetConfiguration.Development : CodeTargetConfiguration;
+						// Need to check for development receipt as we use that for the Engine code in DebugGame
+						UnrealTargetConfiguration EngineConfiguration = (CodeTargetConfiguration == UnrealTargetConfiguration.DebugGame) ? UnrealTargetConfiguration.Development : CodeTargetConfiguration;
 
-							// Android has multiple architecture flavors built without receipts, so use the default arch target instead
-							if (CodeTargetPlatform == UnrealTargetPlatform.Android)
+						// Android has multiple architecture flavors built without receipts, so use the default arch target instead
+						if (CodeTargetPlatform == UnrealTargetPlatform.Android)
+						{
+							FileReference ReceiptFileName = TargetReceipt.GetDefaultPath(new DirectoryReference(OutputEnginePath), CurrentTargetName, CodeTargetPlatform, EngineConfiguration, Architecture);
+							if (FileReference.Exists(ReceiptFileName))
 							{
-								FileReference ReceiptFileName = TargetReceipt.GetDefaultPath(new DirectoryReference(OutputEnginePath), CurrentTargetName, CodeTargetPlatform, EngineConfiguration, Architecture);
-								if (FileReference.Exists(ReceiptFileName))
-								{
-									// Strip the output folder so that this can be used on any machine
-									string RelativeReceiptFileName = ReceiptFileName.MakeRelativeTo(new DirectoryReference(OutputDir));
+								// Strip the output folder so that this can be used on any machine
+								string RelativeReceiptFileName = ReceiptFileName.MakeRelativeTo(new DirectoryReference(OutputDir));
 
-									// Blindly append all of the architecture names
-									if (AllArchNames.Count > 0)
-									{
-										foreach (string Arch in AllArchNames)
-										{
-											InstalledConfigs.Add(new InstalledPlatformInfo.InstalledPlatformConfiguration(CodeTargetConfiguration, CodeTargetPlatform, CurrentTargetType, Arch, RelativeReceiptFileName, ProjectType, bCanBeDisplayed));
-										}
-									}
-									// if for some reason we didn't specify any flavors, just add the default one.
-									else
+								// Blindly append all of the architecture names
+								if (AllArchNames.Count > 0)
+								{
+									foreach (string Arch in AllArchNames)
 									{
-										InstalledConfigs.Add(new InstalledPlatformInfo.InstalledPlatformConfiguration(CodeTargetConfiguration, CodeTargetPlatform, CurrentTargetType, Architecture, RelativeReceiptFileName, ProjectType, bCanBeDisplayed));
+										InstalledConfigs.Add(new InstalledPlatformInfo.InstalledPlatformConfiguration(CodeTargetConfiguration, CodeTargetPlatform, CurrentTargetType, Arch, RelativeReceiptFileName, ProjectType, bCanBeDisplayed));
 									}
 								}
+								// if for some reason we didn't specify any flavors, just add the default one.
+								else
+								{
+									InstalledConfigs.Add(new InstalledPlatformInfo.InstalledPlatformConfiguration(CodeTargetConfiguration, CodeTargetPlatform, CurrentTargetType, Architecture, RelativeReceiptFileName, ProjectType, bCanBeDisplayed));
+								}
 							}
-							// If we're not Android, check the existence of the target receipts for each architecture specified.
-							else
+						}
+						// If we're not Android, check the existence of the target receipts for each architecture specified.
+						else
+						{
+							foreach (string Arch in AllArchNames)
 							{
-								foreach (string Arch in AllArchNames)
+								FileReference ReceiptFileName = TargetReceipt.GetDefaultPath(new DirectoryReference(OutputEnginePath), CurrentTargetName, CodeTargetPlatform, EngineConfiguration, Arch);
+								if (FileReference.Exists(ReceiptFileName))
 								{
-									FileReference ReceiptFileName = TargetReceipt.GetDefaultPath(new DirectoryReference(OutputEnginePath), CurrentTargetName, CodeTargetPlatform, EngineConfiguration, Arch);
-									if (FileReference.Exists(ReceiptFileName))
-									{
-										string RelativeReceiptFileName = ReceiptFileName.MakeRelativeTo(new DirectoryReference(OutputDir));
-										InstalledConfigs.Add(new InstalledPlatformInfo.InstalledPlatformConfiguration(CodeTargetConfiguration, CodeTargetPlatform, CurrentTargetType, Arch, RelativeReceiptFileName, ProjectType, bCanBeDisplayed));
-									}
+									string RelativeReceiptFileName = ReceiptFileName.MakeRelativeTo(new DirectoryReference(OutputDir));
+									InstalledConfigs.Add(new InstalledPlatformInfo.InstalledPlatformConfiguration(CodeTargetConfiguration, CodeTargetPlatform, CurrentTargetType, Arch, RelativeReceiptFileName, ProjectType, bCanBeDisplayed));
 								}
 							}
 						}
@@ -181,5 +175,23 @@ namespace AutomationTool
 			File.AppendAllLines(OutputBaseEnginePath, IniLines);
 			File.SetAttributes(OutputBaseEnginePath, OutputAttributes);
 		}
+
+		/// <summary>
+		/// Parse an argument containing a list of platforms
+		/// </summary>
+		/// <param name="Name">Name of the argument</param>
+		/// <returns>List of platforms</returns>
+		List<UnrealTargetPlatform> ParsePlatformsParamValue(string Name)
+		{
+			string PlatformsString = ParseParamValue(Name);
+			if (String.IsNullOrWhiteSpace(PlatformsString))
+			{
+				return new List<UnrealTargetPlatform>();
+			}
+			else
+			{
+				return PlatformsString.Split(new char[] { ';' }, StringSplitOptions.RemoveEmptyEntries).Select(x => UnrealTargetPlatform.Parse(x)).ToList();
+			}
+		}
 	}
 }
\ No newline at end of file
