Fix issue with missing GoogleGameSDK in Launcher build

diff --git a/Engine/Build/InstalledEngineFilters.xml b/Engine/Build/InstalledEngineFilters.xml
--- a/Engine/Build/InstalledEngineFilters.xml
+++ b/Engine/Build/InstalledEngineFilters.xml
@@ -425,6 +425,7 @@
 		Engine/Source/ThirdParty/Android/...
 		Engine/Source/ThirdParty/AndroidPermission/...
 		Engine/Source/ThirdParty/GoogleARCore/...
+		Engine/Source/ThirdParty/GoogleGameSDK/...
 		Engine/Source/ThirdParty/GoogleVR/...
 		Engine/Source/ThirdParty/GoogleInstantPreview/InstantPreview.apk
 		Engine/Source/ThirdParty/Oculus/LibOVRAudio/LibOVRAudio/lib/*/*.so
