Adding the Datasmith plugin stuff to the Mac platform.

#rnx

diff --git a/Engine/Build/InstalledEngineBuild.xml b/Engine/Build/InstalledEngineBuild.xml
--- a/Engine/Build/InstalledEngineBuild.xml
+++ b/Engine/Build/InstalledEngineBuild.xml
@@ -176,7 +176,7 @@
 		<Node Name="Compile UnrealHeaderTool Mac" Requires="Update Version Files">
 			<Compile Target="UnrealHeaderTool" Platform="Mac" Configuration="Development" Arguments="-precompile -allmodules"/>
 		</Node>
-		<Node Name="Compile UE4Editor Mac" Requires="Compile UnrealHeaderTool Mac" Produces="#UE4Editor Mac;#UE4Editor Mac Unstripped;#UE4Editor Mac Stripped;#UE4Editor Mac Unsigned;#UE4Editor Mac Signed">
+		<Node Name="Compile UE4Editor Mac" Requires="Compile UnrealHeaderTool Mac" Produces="#UE4Editor Mac;#UE4Editor Mac Unstripped;#UE4Editor Mac Stripped;#UE4Editor Mac Unsigned;#UE4Editor Mac Signed;#UE4Editor Mac WithExceptions">
 			<Compile Target="UE4Editor" Platform="Mac" Configuration="DebugGame" Tag="#UE4Editor Mac" Arguments="-precompile -allmodules"/>
 			<Compile Target="UE4Editor" Platform="Mac" Configuration="Development" Tag="#UE4Editor Mac" Arguments="-precompile -allmodules"/>
 
@@ -198,6 +198,8 @@
 				<Copy Files="#UE4Editor Mac Unsigned" From="$(RootDir)" To="$(SavedOutput)" Tag="#Editor_ToSign_Mac"/>
 				<Sign Files="#Editor_ToSign_Mac" Tag="#UE4Editor Mac Signed"/>
 			</Do>
+
+			<Tag Files="#UE4Editor Mac" Except="$(PluginsExceptions)" With="#UE4Editor Mac WithExceptions"/>
 		</Node>
 	</Agent>
 
@@ -738,13 +740,25 @@
 	</Agent>
 
 	<!-- These are plugins that need to be compiled with a full workspace rather than Win64_Licensee. -->
-	<Agent Name="Extra Plugins Agent" Type="Win64">
-		<Node Name="Compile Datasmith Plugins" Requires="Compile UnrealHeaderTool Win64" Produces="#DatasmithPluginBinaries">
+	<Agent Name="Extra Plugins Agent Win64" Type="Win64">
+		<Node Name="Compile Datasmith Plugins Win64" Requires="Compile UnrealHeaderTool Win64" Produces="#DatasmithPluginBinaries Win64">
 			<Do If="$(CompileDatasmithPlugins) == 'True'">
 				<Property Name="DatasmithPlugins" Value="DatasmithCADImporter;DatasmithIFCImporter;DatasmithC4DImporter;AxFImporter;StaticMeshEditorExtension;MDLImporter"/>
 				<ForEach Name="Plugin" Values="$(DatasmithPlugins)">
 					<Compile Target="UE4Editor" Configuration="Development" Platform="Win64" Arguments="-Plugin=&quot;$(RootDir)/Engine/Plugins/Enterprise/$(Plugin)/$(Plugin).uplugin&quot;"/>
-					<Tag Files="$(RootDir)/Engine/Plugins/Enterprise/$(Plugin)/Binaries/..." With="#DatasmithPluginBinaries"/>
+					<Tag Files="$(RootDir)/Engine/Plugins/Enterprise/$(Plugin)/Binaries/..." With="#DatasmithPluginBinaries Win64"/>
+				</ForEach>
+			</Do>
+		</Node>
+	</Agent>
+
+	<Agent Name="Extra Plugins Agent Mac" Type="Mac">
+		<Node Name="Compile Datasmith Plugins Mac" Requires="Compile UnrealHeaderTool Mac" Produces="#DatasmithPluginBinaries Mac">
+			<Do If="$(CompileDatasmithPlugins) == 'True'">
+				<Property Name="DatasmithPlugins" Value="DatasmithCADImporter;DatasmithIFCImporter;DatasmithC4DImporter;AxFImporter;StaticMeshEditorExtension;MDLImporter"/>
+				<ForEach Name="Plugin" Values="$(DatasmithPlugins)">
+					<Compile Target="UE4Editor" Configuration="Development" Platform="Mac" Arguments="-Plugin=&quot;$(RootDir)/Engine/Plugins/Enterprise/$(Plugin)/$(Plugin).uplugin&quot;"/>
+					<Tag Files="$(RootDir)/Engine/Plugins/Enterprise/$(Plugin)/Binaries/..." With="#DatasmithPluginBinaries Mac"/>
 				</ForEach>
 			</Do>
 		</Node>
@@ -767,7 +781,7 @@
 		<Property Name="InstalledRequirements" Value="$(InstalledRequirements);#UE4Game Lumin" If="'$(WithLumin)' == true"/>
 		<Property Name="InstalledRequirements" Value="$(InstalledRequirements);#UE4Game HoloLens" If="'$(WithHoloLens)' == true"/>
 		<Property Name="InstalledRequirements" Value="$(InstalledRequirements);#CompressedDDCWin64" If="'$(WithDDC)' == true"/>
-		<Property Name="InstalledRequirements" Value="$(InstalledRequirements);#Compile Datasmith Plugins;#DatasmithPluginBinaries" If="$(CompileDatasmithPlugins)"/>
+		<Property Name="InstalledRequirements" Value="$(InstalledRequirements);#Compile Datasmith Plugins Win64;#DatasmithPluginBinaries Win64" If="$(CompileDatasmithPlugins)"/>
 		<!-- Feature packs -->
 		<Property Name="InstalledRequirements" Value="$(InstalledRequirements);Make Feature Packs"/>
 
@@ -789,7 +803,7 @@
 			<Tag Files="#UE4Editor Win64 WithExceptions" Except="#UE4Editor Win64 Unstripped;#UE4Editor Win64 Unsigned" With="#Installed Win64"/>
 			<Tag Files="#UE4Editor Win64 Stripped;#UE4Editor Win64 Signed" With="#Saved Output"/>
 			<Tag Files="#Build Tools Win64;#Make Feature Packs" With="#Installed Win64"/>
-			<Tag Files="#DatasmithPluginBinaries" With="#Installed Win64" If="$(CompileDatasmithPlugins)"/>
+			<Tag Files="#DatasmithPluginBinaries Win64" With="#Installed Win64" If="$(CompileDatasmithPlugins)"/>
 
 			<!-- Tag enabled target platforms -->
 			<Do If="'$(WithWin64)' == true">
@@ -938,7 +952,7 @@
 
 	<Agent Name="Installed Build Group Mac" Type="Mac_Licensee">
 		<!-- Build up a list of requirements as some can be optional -->
-		<Property Name="FilterRequirements" Value="#Compile UnrealHeaderTool Mac;#UE4Editor Mac;#UE4Editor Mac Unstripped;#UE4Editor Mac Stripped;#UE4Editor Mac Unsigned;#UE4Editor Mac Signed;#Build Tools Mac;#Build Tools CS Binaries"/>
+		<Property Name="FilterRequirements" Value="#Compile UnrealHeaderTool Mac;#UE4Editor Mac WithExceptions;#UE4Editor Mac Unstripped;#UE4Editor Mac Stripped;#UE4Editor Mac Unsigned;#UE4Editor Mac Signed;#Build Tools Mac;#Build Tools CS Binaries"/>
 
 		<!-- Optional Target Platform requirements -->
 		<Property Name="FilterRequirements" Value="$(FilterRequirements);#UE4Game Mac;#UE4Game Mac Unstripped;#UE4Game Mac Stripped;#UE4Game Mac Unsigned;#UE4Game Mac Signed" If="'$(WithMac)' == true"/>
@@ -947,6 +961,7 @@
 		<Property Name="FilterRequirements" Value="$(FilterRequirements);#UE4Game TVOS;#UE4Game TVOS Unstripped;#UE4Game TVOS Stripped" If="'$(WithTVOS)' == true"/>
 		<Property Name="FilterRequirements" Value="$(FilterRequirements);#UE4Game LuminMac;#UE4Game LuminMac Unstripped;#UE4Game LuminMac Stripped" If="'$(WithLuminMac)' == true"/>
 		<Property Name="FilterRequirements" Value="$(FilterRequirements);#CompressedDDCMac" If="'$(WithDDC)' == true"/>
+		<Property Name="FilterRequirements" Value="$(FilterRequirements);#Compile Datasmith Plugins Mac;#DatasmithPluginBinaries Mac" If="$(CompileDatasmithPlugins)"/>
 
 		<!-- Feature packs -->
 		<Property Name="FilterRequirements" Value="$(FilterRequirements);Make Feature Packs"/>
@@ -965,9 +980,10 @@
 
 			<!-- Tag files that always need to be part of an installed build -->
 			<Tag Files="#Compile UnrealHeaderTool Mac" With="#Installed Mac"/>
-			<Tag Files="#UE4Editor Mac" Except="#UE4Editor Mac Unstripped;#UE4Editor Mac Unsigned" With="#Installed Mac"/>
+			<Tag Files="#UE4Editor Mac WithExceptions" Except="#UE4Editor Mac Unstripped;#UE4Editor Mac Unsigned" With="#Installed Mac"/>
 			<Tag Files="#UE4Editor Mac Stripped;#UE4Editor Mac Signed" With="#Saved Output"/>
 			<Tag Files="#Build Tools Mac;#Make Feature Packs" With="#Installed Mac"/>
+			<Tag Files="#DatasmithPluginBinaries Mac" With="#Installed Mac" If="$(CompileDatasmithPlugins)"/>
 
 			<!-- Tag enabled target platforms -->
 			<Do If="'$(WithMac)' == true">
diff --git a/Engine/Source/Programs/UnrealBuildTool/Platform/Mac/MacToolChain.cs b/Engine/Source/Programs/UnrealBuildTool/Platform/Mac/MacToolChain.cs
--- a/Engine/Source/Programs/UnrealBuildTool/Platform/Mac/MacToolChain.cs
+++ b/Engine/Source/Programs/UnrealBuildTool/Platform/Mac/MacToolChain.cs
@@ -1309,6 +1309,21 @@ namespace UnrealBuildTool
 			}
 		}
 
+		private List<FileItem> DebugInfoFiles = new List<FileItem>();
+
+		public override void FinalizeOutput(ReadOnlyTargetRules Target, TargetMakefile Makefile)
+		{
+			// Re-add any .dSYM files that may have been stripped out.
+			List<string> OutputFiles = Makefile.OutputItems.Select(Item => Path.ChangeExtension(Item.FullName, ".dSYM")).Distinct().ToList();
+			foreach (FileItem DebugItem in DebugInfoFiles)
+			{
+				if(OutputFiles.Any(Item => string.Equals(Item, DebugItem.FullName, StringComparison.InvariantCultureIgnoreCase)))
+				{
+					Makefile.OutputItems.Add(DebugItem);
+				}
+			}
+		}
+
 		public override ICollection<FileItem> PostBuild(FileItem Executable, LinkEnvironment BinaryLinkEnvironment, IActionGraphBuilder Graph)
 		{
 			ICollection<FileItem> OutputFiles = base.PostBuild(Executable, BinaryLinkEnvironment, Graph);
@@ -1332,7 +1347,7 @@ namespace UnrealBuildTool
 				// We want dsyms to be created after all dylib dependencies are fixed. If FixDylibDependencies action was not created yet, save the info for later.
 				if (FixDylibOutputFile != null)
 				{
-					OutputFiles.Add(GenerateDebugInfo(Executable, BinaryLinkEnvironment, Graph));
+					DebugInfoFiles.Add(GenerateDebugInfo(Executable, BinaryLinkEnvironment, Graph));
 				}
 				else
 				{
@@ -1357,7 +1372,7 @@ namespace UnrealBuildTool
 			// Add dsyms that we couldn't add before FixDylibDependencies action was created
 			foreach (FileItem Exe in ExecutablesThatNeedDsyms)
 			{
-				OutputFiles.Add(GenerateDebugInfo(Exe, BinaryLinkEnvironment, Graph));
+				DebugInfoFiles.Add(GenerateDebugInfo(Exe, BinaryLinkEnvironment, Graph));
 			}
 			ExecutablesThatNeedDsyms.Clear();
 
