Add new TP_ProductConfigBP to installed engine build.

diff --git a/Engine/Build/InstalledEngineFilters.xml b/Engine/Build/InstalledEngineFilters.xml
--- a/Engine/Build/InstalledEngineFilters.xml
+++ b/Engine/Build/InstalledEngineFilters.xml
@@ -59,6 +59,7 @@
 		TP_Flying
 		TP_FlyingBP
 		TP_HandheldARBP
+		TP_ProductConfigBP
 		TP_Puzzle
 		TP_PuzzleBP
 		TP_Rolling
@@ -90,6 +91,7 @@
 		Templates/TP_TopDownBP/TP_TopDownBP.uproject
 		Templates/TP_TwinStickBP/TP_TwinStickBP.uproject
 		Templates/TP_VehicleBP/TP_VehicleBP.uproject
+		Templates/TP_ProductConfigBP/TP_ProductConfigBP.uproject
 		Templates/TP_PuzzleBP/TP_PuzzleBP.uproject
 		Templates/TP_2DSideScrollerBP/TP_2DSideScrollerBP.uproject
 		Templates/TP_VehicleAdvBP/TP_VehicleAdvBP.uproject
@@ -137,6 +139,7 @@
 		Templates/TP_Flying/...
 		Templates/TP_FlyingBP/...
 		Templates/TP_HandheldARBP/...
+		Templates/TP_ProductConfigBP/...
 		Templates/TP_Rolling/...
 		Templates/TP_RollingBP/...
 		Templates/TP_SideScroller/...
