Fixing CRC defines for Mac.

#rnx

diff --git a/Engine/Build/InstalledEngineBuild.xml b/Engine/Build/InstalledEngineBuild.xml
--- a/Engine/Build/InstalledEngineBuild.xml
+++ b/Engine/Build/InstalledEngineBuild.xml
@@ -178,7 +178,7 @@
 		</Node>
 		<Node Name="Compile UE4Editor Mac" Requires="Compile UnrealHeaderTool Mac" Produces="#UE4Editor Mac;#UE4Editor Mac Unstripped;#UE4Editor Mac Stripped;#UE4Editor Mac Unsigned;#UE4Editor Mac Signed;#UE4Editor Mac WithExceptions">
 			<Compile Target="UE4Editor" Platform="Mac" Configuration="DebugGame" Tag="#UE4Editor Mac" Arguments="-precompile -allmodules"/>
-			<Compile Target="UE4Editor" Platform="Mac" Configuration="Development" Tag="#UE4Editor Mac" Arguments="-precompile -allmodules"/>
+			<Compile Target="UE4Editor" Platform="Mac" Configuration="Development" Tag="#UE4Editor Mac" Arguments="-precompile -allmodules $(CrashReporterCompileArgs)"/>
 
 			<Tag Files="#UE4Editor Mac" Filter="$(MacStripFilter)" With="#UE4Editor Mac Unstripped"/>
 
@@ -577,8 +577,8 @@
 
 	<Agent Name="Tools Group Mac" Type="Mac_Licensee">
 		<Node Name="Build Tools Mac" Requires="Compile UnrealHeaderTool Mac">
-			<Compile Target="CrashReportClient" Configuration="Shipping" Platform="Mac" Tag="#Build Tools Mac"/>
-			<Compile Target="CrashReportClientEditor" Configuration="Shipping" Platform="Mac" Tag="#Build Tools Mac"/>
+			<Compile Target="CrashReportClient" Configuration="Shipping" Platform="Mac" Arguments="$(CrashReporterCompileArgs)" Tag="#Build Tools Mac"/>
+			<Compile Target="CrashReportClientEditor" Configuration="Shipping" Platform="Mac" Arguments="$(CrashReporterCompileArgs)" Tag="#Build Tools Mac"/>
 			<Compile Target="ShaderCompileWorker" Configuration="Development" Platform="Mac" Tag="#Build Tools Mac"/>
 			<Compile Target="UnrealCEFSubProcess" Configuration="Development" Platform="Mac" Tag="#Build Tools Mac"/>
 			<Compile Target="UnrealInsights" Configuration="Development" Platform="Mac" Tag="#Build Tools Mac"/>
