Removing Datasmith plugins from outputs from Mac editor since they get built separately.

#rnx

diff --git a/Engine/Build/InstalledEngineBuild.xml b/Engine/Build/InstalledEngineBuild.xml
--- a/Engine/Build/InstalledEngineBuild.xml
+++ b/Engine/Build/InstalledEngineBuild.xml
@@ -180,7 +180,7 @@
 			<Compile Target="UE4Editor" Platform="Mac" Configuration="DebugGame" Tag="#UE4Editor Mac" Arguments="-precompile -allmodules"/>
 			<Compile Target="UE4Editor" Platform="Mac" Configuration="Development" Tag="#UE4Editor Mac" Arguments="-precompile -allmodules $(CrashReporterCompileArgs)"/>
 
-			<Tag Files="#UE4Editor Mac" Filter="$(MacStripFilter)" With="#UE4Editor Mac Unstripped"/>
+			<Tag Files="#UE4Editor Mac" Filter="$(MacStripFilter)" Except="$(PluginsExceptions)" With="#UE4Editor Mac Unstripped"/>
 
 			<Switch>
 				<Case If="'$(WithFullDebugInfo)' == true">
@@ -194,7 +194,7 @@
 			</Switch>
 
 			<Do If="'$(SignExecutables)' == true">
-				<Tag Files="#UE4Editor Mac" Filter="$(MacSignFilter)" With="#UE4Editor Mac Unsigned"/>
+				<Tag Files="#UE4Editor Mac" Filter="$(MacSignFilter)" Except="$(PluginsExceptions)" With="#UE4Editor Mac Unsigned"/>
 				<Copy Files="#UE4Editor Mac Unsigned" From="$(RootDir)" To="$(SavedOutput)" Tag="#Editor_ToSign_Mac"/>
 				<Sign Files="#Editor_ToSign_Mac" Tag="#UE4Editor Mac Signed"/>
 			</Do>
