Fixing CRC defines for Mac.

#rnx

diff --git a/Engine/Build/InstalledEngineBuild.xml b/Engine/Build/InstalledEngineBuild.xml
--- a/Engine/Build/InstalledEngineBuild.xml
+++ b/Engine/Build/InstalledEngineBuild.xml
@@ -126,7 +126,7 @@
 	<Option Name="CrashReporterAPIURL" DefaultValue="" Description="The URL to use to talk to the CrashReporterClient API."/>
 	<Option Name="CrashReporterAPIKey" DefaultValue="" Description="The API key to use to talk to the CrashReporterClient API."/>
 	<Property Name="CrashReporterCompileArgs" Value=""/>
-	<Property Name="CrashReporterCompileArgs" Value="-CompilerArguments=&quot;/DCRC_TELEMETRY_URL=$(CrashReporterAPIURL) /DCRC_TELEMETRY_KEY_DEV=$(CrashReporterAPIKey) /DCRC_TELEMETRY_KEY_RELEASE=$(CrashReporterAPIKey)&quot;" If="'$(CrashReporterAPIURL)' != '' And '$(CrashReporterAPIKey)' != ''" />
+	<Property Name="CrashReporterCompileArgs" Value="-Define:CRC_TELEMETRY_URL=$(CrashReporterAPIURL) -Define:CRC_TELEMETRY_KEY_DEV=$(CrashReporterAPIKey) -Define:CRC_TELEMETRY_KEY_RELEASE=$(CrashReporterAPIKey)" If="'$(CrashReporterAPIURL)' != '' And '$(CrashReporterAPIKey)' != ''" />
 
 	<!-- Base editor and tools for compile on Windows -->
 	<Agent Name="Editor Win64" Type="Win64_Licensee">
