Include SoundSwitch in launcher build
#rnx

diff --git a/Engine/Build/InstalledEngineFilters.xml b/Engine/Build/InstalledEngineFilters.xml
--- a/Engine/Build/InstalledEngineFilters.xml
+++ b/Engine/Build/InstalledEngineFilters.xml
@@ -457,6 +457,7 @@
 	<Property Name="CopyIOSFilter" Value="$(MacBuildFiles)">
 		<!-- Build Files -->
 		Engine/Build/IOS/...
+		Engine/Source/ThirdParty/IOS/SoundSwitch/...
 	</Property>
 	<Property Name="CopyIOSFilterWin64" Value="$(CopyIOSFilter)">
 		Engine/Binaries/DotNET/IOS/openssl.exe
