Add BootstrapPackagedGame.app to installed build on Mac.

diff --git a/Engine/Build/InstalledEngineFilters.xml b/Engine/Build/InstalledEngineFilters.xml
--- a/Engine/Build/InstalledEngineFilters.xml
+++ b/Engine/Build/InstalledEngineFilters.xml
@@ -416,6 +416,7 @@
 	
 	<!-- Mac -->
 	<Property Name="CopyMacFilter">
+		Engine/Binaries/Mac/BootstrapPackagedGame.app/...
 	</Property>
 	<Property Name="CopyMacExceptions">
 	</Property>
