Fixing CRC  define behavior on Windows and Mac.

#rnx

diff --git a/Engine/Build/InstalledEngineBuild.xml b/Engine/Build/InstalledEngineBuild.xml
--- a/Engine/Build/InstalledEngineBuild.xml
+++ b/Engine/Build/InstalledEngineBuild.xml
@@ -125,8 +125,10 @@
 
 	<Option Name="CrashReporterAPIURL" DefaultValue="" Description="The URL to use to talk to the CrashReporterClient API."/>
 	<Option Name="CrashReporterAPIKey" DefaultValue="" Description="The API key to use to talk to the CrashReporterClient API."/>
-	<Property Name="CrashReporterCompileArgs" Value=""/>
-	<Property Name="CrashReporterCompileArgs" Value="-Define:CRC_TELEMETRY_URL=$(CrashReporterAPIURL) -Define:CRC_TELEMETRY_KEY_DEV=$(CrashReporterAPIKey) -Define:CRC_TELEMETRY_KEY_RELEASE=$(CrashReporterAPIKey)" If="'$(CrashReporterAPIURL)' != '' And '$(CrashReporterAPIKey)' != ''" />
+	<Property Name="CrashReporterCompileArgsWin" Value=""/>
+	<Property Name="CrashReporterCompileArgsWin" Value="-CompilerArguments=&quot;/DCRC_TELEMETRY_URL=$(CrashReporterAPIURL) /DCRC_TELEMETRY_KEY=$(CrashReporterAPIKey) /DCRC_TELEMETRY_KEY_RELEASE=$(CrashReporterAPIKey)&quot;" If="'$(CrashReporterAPIURL)' != '' And '$(CrashReporterAPIKey)' != ''" />
+	<Property Name="CrashReporterCompileArgsMac" Value=""/>
+	<Property Name="CrashReporterCompileArgsMac" Value="-CompilerArguments=&quot; -DCRC_TELEMETRY_URL=$(CrashReporterAPIURL) -DCRC_TELEMETRY_KEY=$(CrashReporterAPIKey) -DCRC_TELEMETRY_KEY_RELEASE=$(CrashReporterAPIKey)&quot;" If="'$(CrashReporterAPIURL)' != '' And '$(CrashReporterAPIKey)' != ''" />
 
 	<!-- Base editor and tools for compile on Windows -->
 	<Agent Name="Editor Win64" Type="Win64_Licensee">
@@ -139,7 +141,7 @@
 		</Node>
 		<Node Name="Compile UE4Editor Win64" Requires="Compile UnrealHeaderTool Win64" Produces="#UE4Editor Win64;#UE4Editor Win64 Unstripped;#UE4Editor Win64 Stripped;#UE4Editor Win64 Unsigned;#UE4Editor Win64 Signed;#UE4Editor Win64 WithExceptions">
 			<Compile Target="UE4Editor" Platform="Win64" Configuration="DebugGame" Tag="#UE4Editor Win64" Arguments="-precompile -allmodules "/>
-			<Compile Target="UE4Editor" Platform="Win64" Configuration="Development" Tag="#UE4Editor Win64" Arguments="-precompile -allmodules $(CrashReporterCompileArgs)"/>
+			<Compile Target="UE4Editor" Platform="Win64" Configuration="Development" Tag="#UE4Editor Win64" Arguments="-precompile -allmodules $(CrashReporterCompileArgsWin)"/>
 
 			<Do If="$(EmbedSrcSrvInfo)">
 				<!-- Embed source info into the PDB files. Should be done from this machine to ensure that paths are correct. -->
@@ -177,8 +179,8 @@
 			<Compile Target="UnrealHeaderTool" Platform="Mac" Configuration="Development" Arguments="-precompile -allmodules"/>
 		</Node>
 		<Node Name="Compile UE4Editor Mac" Requires="Compile UnrealHeaderTool Mac" Produces="#UE4Editor Mac;#UE4Editor Mac Unstripped;#UE4Editor Mac Stripped;#UE4Editor Mac Unsigned;#UE4Editor Mac Signed;#UE4Editor Mac WithExceptions">
-			<Compile Target="UE4Editor" Platform="Mac" Configuration="DebugGame" Tag="#UE4Editor Mac" Arguments="-precompile -allmodules"/>
-			<Compile Target="UE4Editor" Platform="Mac" Configuration="Development" Tag="#UE4Editor Mac" Arguments="-precompile -allmodules $(CrashReporterCompileArgs)"/>
+			<Compile Target="UE4Editor" Platform="Mac" Configuration="DebugGame" Tag="#UE4Editor Mac" Arguments="-precompile -allmodules $(CrashReporterCompileArgsMac)"/>
+			<Compile Target="UE4Editor" Platform="Mac" Configuration="Development" Tag="#UE4Editor Mac" Arguments="-precompile -allmodules $(CrashReporterCompileArgsMac)"/>
 
 			<Tag Files="#UE4Editor Mac" Filter="$(MacStripFilter)" Except="$(PluginsExceptions)" With="#UE4Editor Mac Unstripped"/>
 
@@ -488,8 +490,8 @@
 
 	<Agent Name="Tools Group Win64" Type="Win64_Licensee">
 		<Node Name="Build Tools Win64" Requires="Compile UnrealHeaderTool Win64;#NotForLicensee Build Tools Win64">
-			<Compile Target="CrashReportClient" Configuration="Shipping" Platform="Win64" Tag="#Build Tools Win64" Arguments="$(CrashReporterCompileArgs)"/>
-			<Compile Target="CrashReportClientEditor" Configuration="Shipping" Platform="Win64" Tag="#Build Tools Win64" Arguments="$(CrashReporterCompileArgs)"/>
+			<Compile Target="CrashReportClient" Configuration="Shipping" Platform="Win64" Tag="#Build Tools Win64" Arguments="$(CrashReporterCompileArgsWin)"/>
+			<Compile Target="CrashReportClientEditor" Configuration="Shipping" Platform="Win64" Tag="#Build Tools Win64" Arguments="$(CrashReporterCompileArgsWin)"/>
 			<Tag Files="#NotForLicensee Build Tools Win64" With="#Build Tools Win64"/>
 			<Compile Target="ShaderCompileWorker" Configuration="Development" Platform="Win64" Tag="#Build Tools Win64"/>
 			<Compile Target="UnrealCEFSubProcess" Configuration="Development" Platform="Win64" Tag="#Build Tools Win64"/>
@@ -577,8 +579,8 @@
 
 	<Agent Name="Tools Group Mac" Type="Mac_Licensee">
 		<Node Name="Build Tools Mac" Requires="Compile UnrealHeaderTool Mac">
-			<Compile Target="CrashReportClient" Configuration="Shipping" Platform="Mac" Arguments="$(CrashReporterCompileArgs)" Tag="#Build Tools Mac"/>
-			<Compile Target="CrashReportClientEditor" Configuration="Shipping" Platform="Mac" Arguments="$(CrashReporterCompileArgs)" Tag="#Build Tools Mac"/>
+			<Compile Target="CrashReportClient" Configuration="Shipping" Platform="Mac" Arguments="$(CrashReporterCompileArgsMac)" Tag="#Build Tools Mac"/>
+			<Compile Target="CrashReportClientEditor" Configuration="Shipping" Platform="Mac" Arguments="$(CrashReporterCompileArgsMac)" Tag="#Build Tools Mac"/>
 			<Compile Target="ShaderCompileWorker" Configuration="Development" Platform="Mac" Tag="#Build Tools Mac"/>
 			<Compile Target="UnrealCEFSubProcess" Configuration="Development" Platform="Mac" Tag="#Build Tools Mac"/>
 			<Compile Target="UnrealInsights" Configuration="Development" Platform="Mac" Tag="#Build Tools Mac"/>
diff --git a/Engine/Source/Programs/UnrealBuildTool/Platform/Mac/MacToolChain.cs b/Engine/Source/Programs/UnrealBuildTool/Platform/Mac/MacToolChain.cs
--- a/Engine/Source/Programs/UnrealBuildTool/Platform/Mac/MacToolChain.cs
+++ b/Engine/Source/Programs/UnrealBuildTool/Platform/Mac/MacToolChain.cs
@@ -11,6 +11,7 @@ using System.Linq;
 using Ionic.Zip;
 using Tools.DotNETCommon;
 using System.Globalization;
+using System.Text.RegularExpressions;
 
 namespace UnrealBuildTool
 {
@@ -537,7 +538,24 @@ namespace UnrealBuildTool
 					CompileAction.ProducedItems.Add(DependencyListFile);
 				}
 
-				string AllArgs = Arguments + FileArguments + CompileEnvironment.AdditionalArguments;
+				string EscapedAdditionalArgs = "";
+				if(!string.IsNullOrWhiteSpace(CompileEnvironment.AdditionalArguments))
+				{
+					foreach(string AdditionalArg in CompileEnvironment.AdditionalArguments.Split(new char[] { ' ' }, StringSplitOptions.RemoveEmptyEntries))
+					{
+						Match DefinitionMatch = Regex.Match(AdditionalArg, "-D\"?(?<Name>.*)=(?<Value>.*)\"?");
+						if (DefinitionMatch.Success)
+						{
+							EscapedAdditionalArgs += string.Format(" -D{0}=\"{1}\"", DefinitionMatch.Groups["Name"].Value, DefinitionMatch.Groups["Value"].Value);
+						}
+						else
+						{
+							EscapedAdditionalArgs += " " + AdditionalArg;
+						}
+					}
+				}
+
+				string AllArgs = Arguments + FileArguments + EscapedAdditionalArgs;
 				string CompilerPath = Settings.ToolchainDir + MacCompiler;
 				
 				// Analyze and then compile using the shell to perform the indirection
diff --git a/Engine/Source/Runtime/CrashReportCore/Private/CrashReportAnalytics.cpp b/Engine/Source/Runtime/CrashReportCore/Private/CrashReportAnalytics.cpp
--- a/Engine/Source/Runtime/CrashReportCore/Private/CrashReportAnalytics.cpp
+++ b/Engine/Source/Runtime/CrashReportCore/Private/CrashReportAnalytics.cpp
@@ -53,7 +53,7 @@ void FCrashReportAnalytics::Initialize()
 			!FEngineBuildSettings::IsInternalBuild();	// Internal Epic build
 
 		FAnalyticsET::Config Config;
-		Config.APIServerET = TEXT(PREPROCESSOR_TO_STRING(CRC_TELEMETRY_URL));
+		Config.APIServerET = FString::Printf(TEXT("https://%s"), TEXT(PREPROCESSOR_TO_STRING(CRC_TELEMETRY_URL)));
 		Config.APIKeyET = bUseReleaseAccount ? TEXT(PREPROCESSOR_TO_STRING(CRC_TELEMETRY_KEY_RELEASE)) : TEXT(PREPROCESSOR_TO_STRING(CRC_TELEMETRY_KEY_DEV));
 
 	#else
