Removing build graph CRC define stuff since it is handled in the .Target.cs file now.

#rnx

diff --git a/Engine/Build/InstalledEngineBuild.xml b/Engine/Build/InstalledEngineBuild.xml
--- a/Engine/Build/InstalledEngineBuild.xml
+++ b/Engine/Build/InstalledEngineBuild.xml
@@ -123,13 +123,6 @@
 	<!-- Include the script that lists files and patterns used when creating an installed build -->
 	<Include Script="InstalledEngineFilters.xml"/>
 
-	<Option Name="CrashReporterAPIURL" DefaultValue="" Description="The URL to use to talk to the CrashReporterClient API."/>
-	<Option Name="CrashReporterAPIKey" DefaultValue="" Description="The API key to use to talk to the CrashReporterClient API."/>
-	<Property Name="CrashReporterCompileArgsWin" Value=""/>
-	<Property Name="CrashReporterCompileArgsWin" Value="-CompilerArguments=&quot;/DCRC_TELEMETRY_URL=\&quot;$(CrashReporterAPIURL)\&quot; /DCRC_TELEMETRY_KEY=$(CrashReporterAPIKey) /DCRC_TELEMETRY_KEY_RELEASE=$(CrashReporterAPIKey)&quot;" If="'$(CrashReporterAPIURL)' != '' And '$(CrashReporterAPIKey)' != ''" />
-	<Property Name="CrashReporterCompileArgsMac" Value=""/>
-	<Property Name="CrashReporterCompileArgsMac" Value="-CompilerArguments=&quot; -DCRC_TELEMETRY_URL=\&quot;$(CrashReporterAPIURL)\&quot; -DCRC_TELEMETRY_KEY=$(CrashReporterAPIKey) -DCRC_TELEMETRY_KEY_RELEASE=$(CrashReporterAPIKey)&quot;" If="'$(CrashReporterAPIURL)' != '' And '$(CrashReporterAPIKey)' != ''" />
-
 	<!-- Base editor and tools for compile on Windows -->
 	<Agent Name="Editor Win64" Type="Win64_Licensee">
 		<Node Name="Update Version Files">
@@ -141,7 +134,7 @@
 		</Node>
 		<Node Name="Compile UE4Editor Win64" Requires="Compile UnrealHeaderTool Win64" Produces="#UE4Editor Win64;#UE4Editor Win64 Unstripped;#UE4Editor Win64 Stripped;#UE4Editor Win64 Unsigned;#UE4Editor Win64 Signed;#UE4Editor Win64 WithExceptions">
 			<Compile Target="UE4Editor" Platform="Win64" Configuration="DebugGame" Tag="#UE4Editor Win64" Arguments="-precompile -allmodules "/>
-			<Compile Target="UE4Editor" Platform="Win64" Configuration="Development" Tag="#UE4Editor Win64" Arguments="-precompile -allmodules $(CrashReporterCompileArgsWin)"/>
+			<Compile Target="UE4Editor" Platform="Win64" Configuration="Development" Tag="#UE4Editor Win64" Arguments="-precompile -allmodules"/>
 
 			<Do If="$(EmbedSrcSrvInfo)">
 				<!-- Embed source info into the PDB files. Should be done from this machine to ensure that paths are correct. -->
@@ -179,8 +172,8 @@
 			<Compile Target="UnrealHeaderTool" Platform="Mac" Configuration="Development" Arguments="-precompile -allmodules"/>
 		</Node>
 		<Node Name="Compile UE4Editor Mac" Requires="Compile UnrealHeaderTool Mac" Produces="#UE4Editor Mac;#UE4Editor Mac Unstripped;#UE4Editor Mac Stripped;#UE4Editor Mac Unsigned;#UE4Editor Mac Signed;#UE4Editor Mac WithExceptions">
-			<Compile Target="UE4Editor" Platform="Mac" Configuration="DebugGame" Tag="#UE4Editor Mac" Arguments="-precompile -allmodules $(CrashReporterCompileArgsMac)"/>
-			<Compile Target="UE4Editor" Platform="Mac" Configuration="Development" Tag="#UE4Editor Mac" Arguments="-precompile -allmodules $(CrashReporterCompileArgsMac)"/>
+			<Compile Target="UE4Editor" Platform="Mac" Configuration="DebugGame" Tag="#UE4Editor Mac" Arguments="-precompile -allmodules"/>
+			<Compile Target="UE4Editor" Platform="Mac" Configuration="Development" Tag="#UE4Editor Mac" Arguments="-precompile -allmodules"/>
 
 			<Tag Files="#UE4Editor Mac" Filter="$(MacStripFilter)" Except="$(PluginsExceptions)" With="#UE4Editor Mac Unstripped"/>
 
@@ -413,7 +406,7 @@
 			<Strip BaseDir="$(RootDir)" OutputDir="$(SavedOutput)" Platform="Lumin" Files="#UE4Game Lumin Unstripped" Tag="#UE4Game Lumin Stripped"/>
 		</Node>
 	</Agent>
-	<Label Name="Lumin" Category="Target Platforms" Requires="Compile UE4Game Lumin"/>
+	<Label Name="Lumin" Category="Target Platforms" Requires="Compile UE4Game Lumin" If="'$(WithLumin)' == true"/>
 
 	<!-- Target Platforms normally built from a Mac machine -->
 	<Agent Name="Target Mac" Type="Mac_Licensee">
@@ -490,8 +483,8 @@
 
 	<Agent Name="Tools Group Win64" Type="Win64_Licensee">
 		<Node Name="Build Tools Win64" Requires="Compile UnrealHeaderTool Win64;#NotForLicensee Build Tools Win64">
-			<Compile Target="CrashReportClient" Configuration="Shipping" Platform="Win64" Tag="#Build Tools Win64" Arguments="$(CrashReporterCompileArgsWin)"/>
-			<Compile Target="CrashReportClientEditor" Configuration="Shipping" Platform="Win64" Tag="#Build Tools Win64" Arguments="$(CrashReporterCompileArgsWin)"/>
+			<Compile Target="CrashReportClient" Configuration="Shipping" Platform="Win64" Tag="#Build Tools Win64"/>
+			<Compile Target="CrashReportClientEditor" Configuration="Shipping" Platform="Win64" Tag="#Build Tools Win64"/>
 			<Tag Files="#NotForLicensee Build Tools Win64" With="#Build Tools Win64"/>
 			<Compile Target="ShaderCompileWorker" Configuration="Development" Platform="Win64" Tag="#Build Tools Win64"/>
 			<Compile Target="UnrealCEFSubProcess" Configuration="Development" Platform="Win64" Tag="#Build Tools Win64"/>
@@ -579,8 +572,8 @@
 
 	<Agent Name="Tools Group Mac" Type="Mac_Licensee">
 		<Node Name="Build Tools Mac" Requires="Compile UnrealHeaderTool Mac">
-			<Compile Target="CrashReportClient" Configuration="Shipping" Platform="Mac" Arguments="$(CrashReporterCompileArgsMac)" Tag="#Build Tools Mac"/>
-			<Compile Target="CrashReportClientEditor" Configuration="Shipping" Platform="Mac" Arguments="$(CrashReporterCompileArgsMac)" Tag="#Build Tools Mac"/>
+			<Compile Target="CrashReportClient" Configuration="Shipping" Platform="Mac" Tag="#Build Tools Mac"/>
+			<Compile Target="CrashReportClientEditor" Configuration="Shipping" Platform="Mac" Tag="#Build Tools Mac"/>
 			<Compile Target="ShaderCompileWorker" Configuration="Development" Platform="Mac" Tag="#Build Tools Mac"/>
 			<Compile Target="UnrealCEFSubProcess" Configuration="Development" Platform="Mac" Tag="#Build Tools Mac"/>
 			<Compile Target="UnrealInsights" Configuration="Development" Platform="Mac" Tag="#Build Tools Mac"/>
