- Removed Win32 from being referenced in:
  - BuildScripts so builders won't compile it
  - UBT platform registration and from being registered in UBT (so it won't show up in VS)
  - PlatformInfo (so it can't triggered from Package/Launch On)

diff --git a/Engine/Build/Graph/Examples/BuildEditorAndTools.xml b/Engine/Build/Graph/Examples/BuildEditorAndTools.xml
--- a/Engine/Build/Graph/Examples/BuildEditorAndTools.xml
+++ b/Engine/Build/Graph/Examples/BuildEditorAndTools.xml
@@ -33,7 +33,7 @@
 
 	<Option Name="EditorTarget" Restrict="[^ ]+" DefaultValue="UE4Editor" Description="Name of the editor target to be built"/>
 	<Option Name="GameTargets" Restrict="[^ ]*" DefaultValue="" Description="Names of game targets to build separated by semicolons, eg. UE4Client;UE4Server"/>
-	<Option Name="TargetPlatforms" Restrict="[^ ]*" DefaultValue="Win64" Description="List of the target platforms to build for, separated by semicolons, eg. Win64;Win32;Android"/>
+	<Option Name="TargetPlatforms" Restrict="[^ ]*" DefaultValue="Win64" Description="List of the target platforms to build for, separated by semicolons, eg. Win64;Android"/>
 	<Option Name="OutputDir" DefaultValue="$(RootDir)\LocalBuilds\$(EditorTarget)Binaries" Description ="Output directory for compiled binaries"/>
 	<Option Name="Licensee" Restrict="true|false" DefaultValue="true" Description="Whether to treat the changelist number as being from a licensee Perforce server"/>
 	<Option Name="Versioned" Restrict="true|false" DefaultValue="$(IsBuildMachine)" Description="Whether to embed changelist number into binaries"/>
diff --git a/Engine/Build/InstalledEngineBuild.xml b/Engine/Build/InstalledEngineBuild.xml
--- a/Engine/Build/InstalledEngineBuild.xml
+++ b/Engine/Build/InstalledEngineBuild.xml
@@ -31,15 +31,11 @@
 	<Property Name="DefaultWithIOS" Value="$(DefaultWithPlatform)"/>
 	<Property Name="DefaultWithIOS" Value="false" If="'$(HostPlatform)' != 'Mac' And !$(AllPlatforms)"/>
 
-	<Property Name="DefaultWithWin32" Value="$(DefaultWithPlatform)"/>
-	<Property Name="DefaultWithWin32" Value="false" If="'$(HostPlatform)' != 'Win64' And !$(AllPlatforms)"/>
-
 	<Property Name="DefaultWithHoloLens" Value="$(DefaultWithPlatform)"/>
 	<Property Name="DefaultWithHoloLens" Value="false" If="'$(HostPlatform)' != 'Win64' And !$(AllPlatforms)"/>
 
 	<!-- Setup default target platforms for Installed build if they haven't been specified on the commandline -->
 	<Option Name="WithWin64" Restrict="true|false" DefaultValue="$(DefaultWithWin64)" Description="Include the Win64 target platform"/>
-	<Option Name="WithWin32" Restrict="true|false" DefaultValue="$(DefaultWithWin32)" Description="Include the Win32 target platform"/>
 	<Option Name="WithMac" Restrict="true|false" DefaultValue="$(DefaultWithMac)" Description="Include the Mac target platform"/>
 	<Option Name="WithAndroid" Restrict="true|false" DefaultValue="$(DefaultWithPlatform)" Description="Include the Android target platform"/>
 	<Option Name="WithIOS" Restrict="true|false" DefaultValue="$(DefaultWithIOS)" Description="Include the iOS target platform"/>
@@ -261,45 +257,6 @@
 	</Agent>
 	<Label Name="Win64" Category="Target Platforms" Requires="Compile UE4Game Win64"/>
 
-	<Agent Name="Target Win32" Type="Win64_Licensee">
-		<Node Name="Compile UE4Game Win32" Requires="Compile UnrealHeaderTool Win64" Produces="#UE4Game Win32;#UE4Game Win32 Unstripped;#UE4Game Win32 Stripped;#UE4Game Win32 Unsigned;#UE4Game Win32 Signed">
-			<ForEach Name="Target" Values="UE4Game;$(OptionalClientTarget)">
-				<ForEach Name="Configuration" Values="$(GameConfigurations)">
-					<Compile Target="$(Target)" Platform="Win32" Configuration="$(Configuration)" Tag="#UE4Game Win32" Arguments="-precompile -allmodules -nolink $(VSCompilerArg) $(TargetDebugInfoArg)"/>
-					<Compile Target="$(Target)" Platform="Win32" Configuration="$(Configuration)" Tag="#UE4Game Win32" Arguments="-precompile $(VSCompilerArg) $(TargetDebugInfoArg)" Clean="false"/>
-				</ForEach>
-			</ForEach>
-
-			<Do If="$(EmbedSrcSrvInfo)">
-				<!-- Embed source info into the PDB files. Should be done from this machine to ensure that paths are correct. -->
-				<Log Message="Embedding source file information into PDB files..."/>
-				<Tag Files="Engine/Source/...;Engine/Plugins/..." Filter="*.c;*.h;*.cpp;*.hpp;*.inl" Except="Engine/Source/ThirdParty/..." With="#SourceFiles"/>
-				<SrcSrv BinaryFiles="#UE4Game Win32" SourceFiles="#SourceFiles" Branch="$(Branch)" Change="$(Change)"/>
-			</Do>
-
-			<Tag Files="#UE4Game Win32" Filter="$(Win32StripFilter)" With="#UE4Game Win32 Unstripped"/>
-
-			<Switch>
-				<Case If="'$(WithFullDebugInfo)' == true">
-					<!-- Copy files for local debug builds -->
-					<Copy Files="#UE4Game Win32 Unstripped" From="$(RootDir)" To="$(SavedOutput)"/>
-				</Case>
-				<Default>
-					<!-- Otherwise, strip the files for distribution -->
-					<Strip BaseDir="$(RootDir)" OutputDir="$(SavedOutput)" Platform="Win32" Files="#UE4Game Win32 Unstripped" Tag="#UE4Game Win32 Stripped"/>
-				</Default>
-			</Switch>
-
-
-			<Do If="'$(SignExecutables)' == true">
-				<Tag Files="#UE4Game Win32" Filter="$(WindowsSignFilter)" With="#UE4Game Win32 Unsigned"/>
-				<Copy Files="#UE4Game Win32 Unsigned" From="$(RootDir)" To="$(SavedOutput)" Tag="#Game_ToSign_Win32"/>
-				<Sign Files="#Game_ToSign_Win32" Tag="#UE4Game Win32 Signed"/>
-			</Do>
-		</Node>
-	</Agent>
-	<Label Name="Win32" Category="Target Platforms" Requires="Compile UE4Game Win32"/>
-
 	<Agent Name="Target Linux" Type="Win64_Licensee">
 		<Node Name="Compile UE4Game Linux" Requires="Compile UnrealHeaderTool $(HostPlatform)" Produces="#UE4Game Linux;#UE4Game Linux Unstripped;#UE4Game Linux Stripped">
 			<ForEach Name="Target" Values="UE4Game;$(OptionalClientTarget);$(OptionalServerTarget)">
@@ -563,11 +520,6 @@
 			<Compile Target="UnrealMultiUserServer" Configuration="Development" Platform="Linux" Tag="#Build Tools Linux"/>
 			<Compile Target="UnrealRecoverySvc" Configuration="Development" Platform="Linux" Tag="#Build Tools Linux"/>
 		</Node>
-		<Node Name="Build Tools Win32" Requires="Compile UnrealHeaderTool Win64">
-			<Compile Target="CrashReportClient" Configuration="Shipping" Platform="Win32" Tag="#Build Tools Win32"/>
-			<Compile Target="UnrealCEFSubProcess" Configuration="Development" Platform="Win32" Tag="#Build Tools Win32"/>
-			<Compile Target="BootstrapPackagedGame" Configuration="Shipping" Platform="Win32" Tag="#Build Tools Win32"/>
-		</Node>
 	</Agent>
 
 	<Agent Name="Tools Group Mac" Type="Mac_Licensee">
@@ -603,7 +555,7 @@
 	</Agent>
 
 	<Property Name="DDCPlatformsWin64" Value="Windows"/>
-	<Property Name="DDCPlatformsWin64" Value="$(DDCPlatformsWin64)+WindowsNoEditor" If="'$(WithWin64)' == true Or '$(WithWin32)' == true"/>
+	<Property Name="DDCPlatformsWin64" Value="$(DDCPlatformsWin64)+WindowsNoEditor" If="'$(WithWin64)' == true"/>
 	<Do If="'$(HostPlatformDDCOnly)' == false">
 		<Property Name="DDCPlatformsWin64" Value="$(DDCPlatformsWin64)+Android_ATC" If="'$(WithAndroid)' == true"/>
 		<Property Name="DDCPlatformsWin64" Value="$(DDCPlatformsWin64)+IOS" If="'$(WithIOS)' == true"/>
@@ -672,7 +624,7 @@
 			<TagReceipt Files="#TargetReceipts" RuntimeDependencies="true" With="#ToCopy"/>
 			<Tag Files="Engine/Content/..." Except="*.psd;*.pdn;*.fbx;*.po" With="#ToCopy"/>
 			<Tag Files="Engine/Config/..." Except="*.vdf" With="#ToCopy"/>
-			<Tag Files="Engine/Plugins/....uplugin;Engine/Plugins/.../Config/...;Engine/Plugins/.../Content/...;Engine/Plugins/.../Resources/...;Engine/Plugins/.../Shaders/...;Engine/Plugins/.../Templates/..." Except="Engine/Plugins/Runtime/TwitchLiveStreaming/...;Engine/Plugins/.../Binaries/Win64/...;Engine/Plugins/.../Binaries/Win32/..." With="#ToCopy"/>
+			<Tag Files="Engine/Plugins/....uplugin;Engine/Plugins/.../Config/...;Engine/Plugins/.../Content/...;Engine/Plugins/.../Resources/...;Engine/Plugins/.../Shaders/...;Engine/Plugins/.../Templates/..." Except="Engine/Plugins/Runtime/TwitchLiveStreaming/...;Engine/Plugins/.../Binaries/Win64/..." With="#ToCopy"/>
 			<Tag Files="Engine/Shaders/..." With="#ToCopy"/>
 			<Tag Files="Engine/Binaries/DotNET/Ionic.Zip.Reduced.dll" With="#ToCopy"/>
 			<Tag Files="Engine/Binaries/DotNET/OneSky.dll" With="#ToCopy"/>
@@ -715,7 +667,7 @@
 			<TagReceipt Files="#TargetReceipts" RuntimeDependencies="true" With="#ToCopy"/>
 			<Tag Files="Engine/Content/..." Except="*.psd;*.pdn;*.fbx;*.po" With="#ToCopy"/>
 			<Tag Files="Engine/Config/..." Except="*.vdf" With="#ToCopy"/>
-			<Tag Files="Engine/Plugins/....uplugin;Engine/Plugins/.../Config/...;Engine/Plugins/.../Content/...;Engine/Plugins/.../Resources/...;Engine/Plugins/.../Shaders/...;Engine/Plugins/.../Templates/..." Except="Engine/Plugins/Runtime/TwitchLiveStreaming/...;Engine/Plugins/.../Binaries/Win32/...;Engine/Plugins/.../Binaries/Mac/..." With="#ToCopy"/>
+			<Tag Files="Engine/Plugins/....uplugin;Engine/Plugins/.../Config/...;Engine/Plugins/.../Content/...;Engine/Plugins/.../Resources/...;Engine/Plugins/.../Shaders/...;Engine/Plugins/.../Templates/..." Except="Engine/Plugins/Runtime/TwitchLiveStreaming/...;Engine/Plugins/.../Binaries/Mac/..." With="#ToCopy"/>
 			<Tag Files="Engine/Shaders/..." With="#ToCopy"/>
 			<Tag Files="Engine/Binaries/DotNET/Ionic.Zip.Reduced.dll" With="#ToCopy"/>
 			<Tag Files="Engine/Binaries/DotNET/OneSky.dll" With="#ToCopy"/>
@@ -764,7 +716,6 @@
 		<Property Name="InstalledRequirements" Value="Compile UnrealHeaderTool Win64;#UE4Editor Win64 WithExceptions;#UE4Editor Win64 Unstripped;#UE4Editor Win64 Stripped;#UE4Editor Win64 Unsigned;#UE4Editor Win64 Signed;#Build Tools Win64;#Build Tools CS Binaries"/>
 		<!-- Optional Target Platform requirements -->
 		<Property Name="InstalledRequirements" Value="$(InstalledRequirements);#UE4Game Win64;#UE4Game Win64 Unstripped;#UE4Game Win64 Stripped;#UE4Game Win64 Unsigned;#UE4Game Win64 Signed" If="'$(WithWin64)' == true"/>
-		<Property Name="InstalledRequirements" Value="$(InstalledRequirements);#UE4Game Win32;#UE4Game Win32 Unstripped;#UE4Game Win32 Stripped;#UE4Game Win32 Unsigned;#UE4Game Win32 Signed;#Build Tools Win32" If="'$(WithWin32)' == true"/>
 		<Property Name="InstalledRequirements" Value="$(InstalledRequirements);#UE4Game Android;#UE4Game Android Unstripped;#UE4Game Android Stripped" If="'$(WithAndroid)' == true"/>
 		<Property Name="InstalledRequirements" Value="$(InstalledRequirements);#UE4Game IOS;#UE4Game IOS Unstripped;#UE4Game IOS Stripped;#Compile UnrealHeaderTool Mac" If="'$(WithIOS)' == true"/>
 		<Property Name="InstalledRequirements" Value="$(InstalledRequirements);#UE4Game TVOS;#UE4Game TVOS Unstripped;#UE4Game TVOS Stripped" If="'$(WithTVOS)' == true"/>
@@ -808,13 +759,6 @@
 				<Property Name="CopyInstalledFilter" Value="$(CopyInstalledFilter);$(CopyWin64Filter)"/>
 				<Property Name="CopyInstalledExceptions" Value="$(CopyInstalledExceptions);$(CopyWin64Exceptions)"/>
 			</Do>
-			<Do If="'$(WithWin32)' == true">
-				<Tag Files="#UE4Game Win32;#Build Tools Win32" Except="#UE4Game Win32 Unstripped;#UE4Game Win32 Unsigned" With="#Installed Win64"/>
-				<Tag Files="#UE4Game Win32 Stripped;#UE4Game Win32 Signed" With="#Saved Output"/>
-				<Property Name="Platforms" Value="$(Platforms)Win32;"/>
-				<Property Name="CopyInstalledFilter" Value="$(CopyInstalledFilter);$(CopyWin32Filter)"/>
-				<Property Name="CopyInstalledExceptions" Value="$(CopyInstalledExceptions);$(CopyWin32Exceptions)"/>
-			</Do>
 			<Do If="'$(WithAndroid)' == true">
 				<Tag Files="#UE4Game Android" Except="#UE4Game Android Unstripped" With="#Installed Win64"/>
 				<Tag Files="#UE4Game Android Stripped" With="#Saved Output"/>
@@ -896,12 +840,10 @@
 
 			<!-- Tag Win64 files that need to be stripped -->
 			<Tag Files="#Compile UnrealHeaderTool Win64;#Build Tools Win64" Filter="$(Win64StripFilter)" Except="$(Win64StripExceptions);$(CopyInstalledExceptions)" With="#Strip Tools for Installed Win64"/>
-			<Tag Files="#Build Tools Win32" Filter="$(Win32StripFilter)" Except="$(CopyInstalledExceptions)" With="#Strip Tools for Installed Win64" If="$(WithWin32)"/>
 
 			<Do If="'$(SignExecutables)' == true">
 				<!-- Tag files that need to be signed -->
 				<Tag Files="#Compile UnrealHeaderTool Win64;#Build Tools Win64" Filter="$(WindowsSignFilter)" Except="$(CopyInstalledExceptions);$(SignExceptions)" With="#Sign Tools for Installed Win64"/>
-				<Tag Files="#Build Tools Win32" Filter="$(WindowsSignFilter)" Except="$(CopyInstalledExceptions);$(SignExceptions)" With="#Sign Tools for Installed Win64" If="$(WithWin32)"/>
 			</Do>
 
 			<!-- Tag files that need to be copied directly -->
diff --git a/Engine/Build/InstalledEngineFilters.xml b/Engine/Build/InstalledEngineFilters.xml
--- a/Engine/Build/InstalledEngineFilters.xml
+++ b/Engine/Build/InstalledEngineFilters.xml
@@ -27,7 +27,6 @@
 	<!-- List of file types to be stripped and signed for different platforms -->
 	<Property Name="Win64StripFilter" Value="*.pdb"/>
 	<Property Name="Win64StripExceptions" Value="Engine\Binaries\Win64\UE4Editor*.pdb;Engine\Plugins\...\Binaries\Win64\UE4Editor*.pdb"/>
-	<Property Name="Win32StripFilter" Value="*.pdb"/>
 	<Property Name="MacStripFilter" Value="*.a;"/>
 	<Property Name="AndroidStripFilter" Value="*.a;*.so"/>
 	<Property Name="IOSStripFilter" Value="*.a;"/>
@@ -332,11 +331,9 @@
 	<Property Name="CopyEditorExceptionsMac" Value="$(CopyEditorExceptions)">
 		<!-- Exclude Windows binaries on Mac here -->
 		Engine/Plugins/.../Binaries/Win64/...
-		Engine/Plugins/.../Binaries/Win32/...
 
 		<!-- Don't want these folders, even if they're part of Windows tools -->
 		Engine/Binaries/Win64/...
-		Engine/Binaries/Win32/...
 	</Property>
 
 	<!-- Linux editor specific dependencies -->
@@ -395,24 +392,6 @@
 	</Property>
 	<Property Name="CopyWin64Exceptions">
 	</Property>
-
-	<!-- Win32 -->
-	<Property Name="CopyWin32Filter">
-		<!-- Natvis helpers for live coding -->
-		Engine/Extras/NatvisHelpers/Win32/NatvisHelpers.lib
-		
-		<!-- Oculus Libraries & Build Files -->
-		Engine/Source/ThirdParty/Oculus/LibOVRAudio/LibOVRAudio/lib/win32/*.lib
-		Engine/Source/ThirdParty/Oculus/LibOVRAudio/*.tps
-		Engine/Source/ThirdParty/Oculus/LibOVRAvatar/LibOVRAvatar/lib/win32/*.lib
-		Engine/Source/ThirdParty/Oculus/LibOVRAvatar/*.tps
-		Engine/Source/ThirdParty/Oculus/LibOVRPlatform/LibOVRPlatform/lib/*.lib
-		Engine/Source/ThirdParty/Oculus/LibOVRPlatform/*.tps
-		Engine/Source/ThirdParty/Oculus/OVRPlugin/OVRPlugin/lib/Win32/*.lib
-		Engine/Source/ThirdParty/Oculus/OVRPlugin/*.tps
-	</Property>
-	<Property Name="CopyWin32Exceptions">
-	</Property>
 	
 	<!-- Mac -->
 	<Property Name="CopyMacFilter">
diff --git a/Engine/Build/PhysXBuild.xml b/Engine/Build/PhysXBuild.xml
--- a/Engine/Build/PhysXBuild.xml
+++ b/Engine/Build/PhysXBuild.xml
@@ -11,7 +11,6 @@
 	<Option Name="TargetLibs" DefaultValue="PhysX+APEX" Description="Target libs"/>
 	<Option Name="TargetConfigs" DefaultValue="profile+release+checked" Description="Target Configs"/>
 	<Option Name="TargetWindowsCompilers" DefaultValue="VisualStudio2015" Description="Target Compilers"/>
-	<Option Name="WithWin32" DefaultValue="false" Description="Compile for Win32"/>
 	<Option Name="WithWin64" DefaultValue="false" Description="Compile for Win64"/>
 	<Option Name="WithHoloLens" DefaultValue="false" Description="Compile for HoloLens"/>
 	<Option Name="WithPS4" DefaultValue="false" Description="Compile for PS4"/>
@@ -38,11 +37,6 @@
 	<Property Name="OtherArgs" Value="$(OtherArgs) -SkipSubmit" If="$(SkipSubmit)"/>
 	<Property Name="OtherArgs" Value="$(OtherArgs) -Robomerge=$(Robomerge)" If="'$(Robomerge)' != ''"/>
 	
-	<Agent Name="PhysX Win32" Type="EngineWin64;PhysX_CompileWin64;Win64">
-		<Node Name="Compile PhysX Win32">
-			<Command Name="BuildPhysX" Arguments="-TargetPlatforms=Win32 -TargetConfigs=&quot;$(TargetConfigs)&quot; -TargetWindowsCompilers=&quot;$(TargetWindowsCompilers)&quot; $(OtherArgs)"/>
-		</Node>
-	</Agent>
 	<Agent Name="PhysX Win64" Type="EngineWin64;PhysX_CompileWin64;Win64">
 		<Node Name="Compile PhysX Win64">
 			<Command Name="BuildPhysX" Arguments="-TargetPlatforms=Win64 -TargetConfigs=&quot;$(TargetConfigs)&quot; -TargetWindowsCompilers=&quot;$(TargetWindowsCompilers)&quot; $(OtherArgs)"/>
@@ -126,7 +120,6 @@
 	</Agent>
 
 	<Property Name="PhysXTargets" Value=""/>
-	<Property Name="PhysXTargets" Value="$(PhysXTargets);Compile PhysX Win32" If="$(WithWin32)"/>
 	<Property Name="PhysXTargets" Value="$(PhysXTargets);Compile PhysX Win64" If="$(WithWin64)"/>
 	<Property Name="PhysXTargets" Value="$(PhysXTargets);Compile PhysX HoloLens" If="$(WithHoloLens)"/>
 	<Property Name="PhysXTargets" Value="$(PhysXTargets);Compile PhysX PS4" If="$(WithPS4)"/>
@@ -148,7 +141,7 @@
 	
 	<Property Name="AllEmailTargets" Value="[REDACTED_EMAIL];[REDACTED_EMAIL];[REDACTED_EMAIL];[REDACTED_EMAIL];[REDACTED_EMAIL];[REDACTED_EMAIL];[REDACTED_EMAIL]"/>
 	
-	<Notify Nodes="Compile PhysX Win32;Compile PhysX Win64" Users="$(AllEmailTargets)"/>
+	<Notify Nodes="Compile PhysX Win64" Users="$(AllEmailTargets)"/>
 	<Notify Nodes="Compile PhysX PS4" Users="$(AllEmailTargets);[REDACTED_EMAIL]"/>
 	<Notify Nodes="Compile PhysX XboxOne" Users="$(AllEmailTargets);[REDACTED_EMAIL]"/>
 	<Notify Nodes="Compile PhysX Linux;Compile PhysX Linux x86;Compile PhysX Linux ARM32;Compile PhysX Linux ARM64" Users="$(AllEmailTargets);[REDACTED_EMAIL]"/>
diff --git a/Engine/Config/Windows/DataDrivenPlatformInfo.ini b/Engine/Config/Windows/DataDrivenPlatformInfo.ini
--- a/Engine/Config/Windows/DataDrivenPlatformInfo.ini
+++ b/Engine/Config/Windows/DataDrivenPlatformInfo.ini
@@ -1,36 +1,15 @@
 [DataDrivenPlatformInfo]
 bIsConfidential=false
 AudioCompressionSettingsIniSectionName=/Script/WindowsTargetPlatform.WindowsTargetSettings
-Freezing_b32Bit=ini:Engine:Windows:/Script/WindowsTargetPlatform.WindowsTargetSettings:bTarget32Bit
-Freezing_bWithRayTracing=!ini:Engine:Windows:/Script/WindowsTargetPlatform.WindowsTargetSettings:bTarget32Bit
+Freezing_b32Bit=false
+Freezing_bWithRayTracing=true
 Freezing_MaxFieldAlignment32=4
 Freezing_MaxFieldAlignment64=8
 Freezing_bAlignBases=true
 
-[PlatformInfo WindowsNoEditorWin32]
-TargetPlatformName=WindowsNoEditor
-DisplayName=Windows (32-bit)
-PlatformType=Game
-PlatformFlags=None
-NormalIconPath=Launcher/Windows/Platform_WindowsNoEditor_24x
-LargeIconPath=Launcher/Windows/Platform_WindowsNoEditor_128x
-XLargeIconPath=
-UATCommandLine=-targetplatform=Win32
-AutoSDKPath=
-TutorialPath=/Engine/Tutorial/Installation/InstallingVisualStudioTutorial.InstallingVisualStudioTutorial
-Windows:bIsEnabled=true
-Mac:bIsEnabled=false
-Linux:bIsEnabled=false
-BinariesDirectoryName=Win32
-IniPlatformName=Windows
-bUsesHostCompiler=true
-bUATClosesAfterLaunch=false
-UBTTargetID=Win32
-PlatformGroupName=Desktop
-
 [PlatformInfo WindowsNoEditor]
 TargetPlatformName=WindowsNoEditor
-DisplayName=Windows (64-bit)
+DisplayName=Windows
 PlatformType=Game
 PlatformFlags=None
 NormalIconPath=Launcher/Windows/Platform_WindowsNoEditor_24x
diff --git a/Engine/Source/Developer/LauncherServices/Private/Launcher/LauncherWorker.cpp b/Engine/Source/Developer/LauncherServices/Private/Launcher/LauncherWorker.cpp
--- a/Engine/Source/Developer/LauncherServices/Private/Launcher/LauncherWorker.cpp
+++ b/Engine/Source/Developer/LauncherServices/Private/Launcher/LauncherWorker.cpp
@@ -337,27 +337,7 @@ FString FLauncherWorker::CreateUATCommand( const ILauncherProfileRef& InProfile,
 			{
 				Platforms += TEXT("+LinuxAArch64");
 			}
-			else if (PlatformInfo->TargetPlatformName == FName("WindowsNoEditor"))
-			{
-				// find out if the project is targeting 32bit
-				FConfigFile ProjectEngineConfig;
-				FString ProjectDir = FPaths::GetPath(InProfile->GetProjectPath());
-				FConfigCacheIni::LoadExternalIniFile(ProjectEngineConfig, TEXT("Engine"), *FPaths::EngineConfigDir(), *FPaths::Combine(ProjectDir, TEXT("Config/")), true, TEXT("Windows"));
-				bool bTarget32Bit = false;
-				ProjectEngineConfig.GetBool(TEXT("/Script/WindowsTargetPlatform.WindowsTargetSettings"), TEXT("bTarget32Bit"), bTarget32Bit);
-
-				// normally this would get the 64bit one, so if we need the 32-bit one, swap it out
-				if (bTarget32Bit)
-				{
-					PlatformInfo = PlatformInfo::FindPlatformInfo(TEXT("WindowsNoEditorWin32"));
-					check(PlatformInfo != nullptr);
-				}
-			
-				// if target wants 32-bit, use 32-bit
-				Platforms += TEXT("+");
-				Platforms += PlatformInfo->UBTTargetId.ToString();
-			}
-			else if (PlatformInfo->TargetPlatformName == FName("Windows") || PlatformInfo->TargetPlatformName == FName("WindowsClient"))
+			else if (PlatformInfo->TargetPlatformName == FName("WindowsNoEditor") || PlatformInfo->TargetPlatformName == FName("Windows") || PlatformInfo->TargetPlatformName == FName("WindowsClient"))
 			{
 				Platforms += TEXT("+Win64");
 			}
diff --git a/Engine/Source/Developer/Windows/WindowsTargetPlatform/Classes/WindowsTargetSettings.h b/Engine/Source/Developer/Windows/WindowsTargetPlatform/Classes/WindowsTargetSettings.h
--- a/Engine/Source/Developer/Windows/WindowsTargetPlatform/Classes/WindowsTargetSettings.h
+++ b/Engine/Source/Developer/Windows/WindowsTargetPlatform/Classes/WindowsTargetSettings.h
@@ -73,12 +73,6 @@ public:
 	UPROPERTY(EditAnywhere, config, Category="OS Info", Meta=(DisplayName = "Minimum OS Version"))
 	EMinimumSupportedOS MinimumOSVersion;
 
-	/**
-	 * Determines if data is cooked for 32bit (or 64bit) exe
-	 */
-	UPROPERTY(EditAnywhere, config, Category = "OS Info", Meta = (DisplayName = "Target 32bit", ConfigRestartRequired=true))
-	bool bTarget32Bit;
-
 	/** Sample rate to run the audio mixer with. */
 	UPROPERTY(config, EditAnywhere, Category = "Audio", Meta = (DisplayName = "Audio Mixer Sample Rate"))
 	int32 AudioSampleRate;
diff --git a/Engine/Source/Developer/Windows/WindowsTargetPlatform/Private/WindowsTargetPlatformClasses.cpp b/Engine/Source/Developer/Windows/WindowsTargetPlatform/Private/WindowsTargetPlatformClasses.cpp
--- a/Engine/Source/Developer/Windows/WindowsTargetPlatform/Private/WindowsTargetPlatformClasses.cpp
+++ b/Engine/Source/Developer/Windows/WindowsTargetPlatform/Private/WindowsTargetPlatformClasses.cpp
@@ -9,7 +9,6 @@ UWindowsTargetSettings::UWindowsTargetSettings( const FObjectInitializer& Object
 	: Super(ObjectInitializer)
 {
 	MinimumOSVersion = EMinimumSupportedOS::MSOS_Vista;
-	bTarget32Bit = false;
 
 	// Default windows settings
 	AudioSampleRate = 48000;
diff --git a/Engine/Source/Editor/MainFrame/Private/Menus/PackageProjectMenu.h b/Engine/Source/Editor/MainFrame/Private/Menus/PackageProjectMenu.h
--- a/Engine/Source/Editor/MainFrame/Private/Menus/PackageProjectMenu.h
+++ b/Engine/Source/Editor/MainFrame/Private/Menus/PackageProjectMenu.h
@@ -69,15 +69,6 @@ public:
 				continue;
 			}
 
-			// windows target version needs special handling
-			bool bTargetWin32 = false;
-			GConfig->GetBool(TEXT("/Script/WindowsTargetPlatform.WindowsTargetSettings"), TEXT("bTarget32Bit"), bTargetWin32, GEngineIni);
-			// if we are targeting 32-bit, only show 32-bit option, for Windows
-			if ((VanillaPlatform.PlatformInfo->UBTTargetId == TEXT("Win32") && !bTargetWin32) || (VanillaPlatform.PlatformInfo->UBTTargetId == TEXT("Win64") && bTargetWin32))
-			{
-				continue;
-			}
-
 			// Check if this platform has a submenu entry
 			if (VanillaPlatform.PlatformInfo->PlatformSubMenu != NAME_None)
 			{
diff --git a/Engine/Source/Programs/UnrealBuildTool/Platform/Windows/UEBuildWindows.cs b/Engine/Source/Programs/UnrealBuildTool/Platform/Windows/UEBuildWindows.cs
--- a/Engine/Source/Programs/UnrealBuildTool/Platform/Windows/UEBuildWindows.cs
+++ b/Engine/Source/Programs/UnrealBuildTool/Platform/Windows/UEBuildWindows.cs
@@ -2330,16 +2330,11 @@ namespace UnrealBuildTool
 			WindowsPlatformSDK SDK = new WindowsPlatformSDK();
 			SDK.ManageAndValidateSDK();
 
-			// Register this build platform for both Win64 and Win32
+			// Register this build platform for Win64 (no more Win32)
 			UEBuildPlatform.RegisterBuildPlatform(new WindowsPlatform(UnrealTargetPlatform.Win64, SDK));
 			UEBuildPlatform.RegisterPlatformWithGroup(UnrealTargetPlatform.Win64, UnrealPlatformGroup.Windows);
 			UEBuildPlatform.RegisterPlatformWithGroup(UnrealTargetPlatform.Win64, UnrealPlatformGroup.Microsoft);
 			UEBuildPlatform.RegisterPlatformWithGroup(UnrealTargetPlatform.Win64, UnrealPlatformGroup.Desktop);
-
-			UEBuildPlatform.RegisterBuildPlatform(new WindowsPlatform(UnrealTargetPlatform.Win32, SDK));
-			UEBuildPlatform.RegisterPlatformWithGroup(UnrealTargetPlatform.Win32, UnrealPlatformGroup.Windows);
-			UEBuildPlatform.RegisterPlatformWithGroup(UnrealTargetPlatform.Win32, UnrealPlatformGroup.Microsoft);
-			UEBuildPlatform.RegisterPlatformWithGroup(UnrealTargetPlatform.Win32, UnrealPlatformGroup.Desktop);
 		}
 	}
 }
