Adding UGS ini.

#rnx

diff --git a/Engine/Build/UnrealGameSync.ini b/Engine/Build/UnrealGameSync.ini
new file mode 100644
--- /dev/null
+++ b/Engine/Build/UnrealGameSync.ini
@@ -0,0 +1,2 @@
+[//UE5/Main/UE4Games.uprojectdirs]
+ZippedBinariesPath=//UE5/Private-Binaryes/++UE5+Main.zip
\ No newline at end of file
