Fixing UGS ini.

#rnx

diff --git a/Engine/Build/UnrealGameSync.ini b/Engine/Build/UnrealGameSync.ini
--- a/Engine/Build/UnrealGameSync.ini
+++ b/Engine/Build/UnrealGameSync.ini
@@ -1,2 +1,2 @@
 [//UE5/Main/UE4Games.uprojectdirs]
-ZippedBinariesPath=//UE5/Private-Binaryes/++UE5+Main.zip
\ No newline at end of file
+ZippedBinariesPath=//UE5/Private-Binaries/++UE5+Main.zip
\ No newline at end of file
