Adding cloud builders to acceptable agent types for Win64 work.

#rnx

diff --git a/Engine/Build/InstalledEngineBuild.xml b/Engine/Build/InstalledEngineBuild.xml
--- a/Engine/Build/InstalledEngineBuild.xml
+++ b/Engine/Build/InstalledEngineBuild.xml
@@ -424,7 +424,7 @@
 	</Agent>
 	<Aggregate Name="TargetPlatforms_LuminMac" Label="Platforms/LuminMac" Requires="Compile UE4Game LuminMac"/>
 
-	<Agent Name="NotForLicensee Tools Group Win64" Type="Win64">
+	<Agent Name="NotForLicensee Tools Group Win64" Type="Win64AWS;Win64">
 		<Node Name="Build NotForLicensee Tools Win64" Requires="Compile UnrealHeaderTool Win64" Produces="#NotForLicensee Build Tools Win64">
 			<Compile Target="DatasmithCADWorker" Configuration="Development" Platform="Win64" Tag="#NotForLicensee Build Tools Win64"/>
 		</Node>
@@ -679,7 +679,7 @@
 		<Property Name="DDCPlatformsLinux" Value="$(DDCPlatformsLinux)+LinuxAArch64" If="'$(WithLinuxAArch64)' == true"/>
 	</Do>
 
-	<Agent Name="DDC Group Linux" Type="Linux;Win64">
+	<Agent Name="DDC Group Linux" Type="Win64AWS;Linux;Win64">
 		<!-- Additionally require the Win64 editor binaries if building on win64 -->
 		<Property Name="DDCLinuxRequires" Value="#UE4Editor Linux;Build Tools Linux;#Build Tools CS Binaries"/>
 		<Property Name="DDCLinuxRequires" Value="$(DDCLinuxRequires);#UE4Editor Win64;#Build Tools Win64" If="'$(HostPlatform)' == 'Win64'"/>
@@ -717,7 +717,7 @@
 	</Agent>
 
 	<!-- These are plugins that need to be compiled with a full workspace rather than Win64_Licensee. -->
-	<Agent Name="Extra Plugins Agent Win64" Type="Win64">
+	<Agent Name="Extra Plugins Agent Win64" Type="Win64AWS;Win64">
 		<Node Name="Compile Datasmith Plugins Win64" Requires="Compile UnrealHeaderTool Win64" Produces="#DatasmithPluginBinaries Win64">
 			<Do If="$(CompileDatasmithPlugins) == 'True'">
 				<Property Name="DatasmithPlugins" Value="DatasmithCADImporter;DatasmithIFCImporter;DatasmithC4DImporter;AxFImporter;StaticMeshEditorExtension;MDLImporter"/>
