Rename UE4Editor to UE5Editor, and generated solution/projects from UE4 -> UE5.

UE4Game/Client/Server targets left intact for the time being.

diff --git a/Engine/Build/BatchFiles/Linux/QASmokeManual.sh b/Engine/Build/BatchFiles/Linux/QASmokeManual.sh
--- a/Engine/Build/BatchFiles/Linux/QASmokeManual.sh
+++ b/Engine/Build/BatchFiles/Linux/QASmokeManual.sh
@@ -97,7 +97,7 @@ if [ -z "$1" ]; then
   # build editor targets, roughly ascending by build time
   Run BuildTarget ShaderCompileWorker
   Run BuildTarget UnrealFrontend
-  Run BuildTarget UE4Editor
+  Run BuildTarget UE5Editor
   popd >> /dev/null
 
 else
diff --git a/Engine/Build/BatchFiles/Linux/README.md b/Engine/Build/BatchFiles/Linux/README.md
--- a/Engine/Build/BatchFiles/Linux/README.md
+++ b/Engine/Build/BatchFiles/Linux/README.md
@@ -80,7 +80,7 @@ GenerateProjectFiles.sh produces a number of "project" files, including Makefile
 project in your favorite IDE and qmake project file. Both QtCreator and KDevelop 4.6+ are known to handle the project well 
 (although the latter takes about 3-4 GB of resident RAM to load the project).
 
-The targets match the name of the resulting binary, e.g. UE4Editor-Linux-Debug or UE4Game. You can build them
+The targets match the name of the resulting binary, e.g. UE5Editor-Linux-Debug or UE4Game. You can build them
 by just typing make <target> in the engine's root folder.
 
 4.8 and later versions have simplified building the editor by providing makefile targets "StandardSet" and "DebugSet", the former
@@ -90,28 +90,28 @@ being the default. You can now just type
 
 to build the editor. Alternatively, build the following targets:
 
-    make CrashReportClient ShaderCompileWorker UnrealLightmass UnrealPak UE4Editor
+    make CrashReportClient ShaderCompileWorker UnrealLightmass UnrealPak UE5Editor
 
 If you intend to develop the editor, you can build a debug configuration of it:
 
-    make UE4Editor-Linux-Debug
+    make UE5Editor-Linux-Debug
 
 (note that it will still use development ShaderCompileWorker / UnrealLightmass). This
 configuration runs much slower.
 
 If you want to rebuild the editor from scatch, you can use
 
-    make UE4Editor ARGS="-clean" && make UE4Editor
+    make UE5Editor ARGS="-clean" && make UE5Editor
 
 In order to run it:
 
     cd Engine/Binaries/Linux/
-    ./UE4Editor
+    ./UE5Editor
 
 Or, if you want to start it with a specific project:
 
     cd Engine/Binaries/Linux/
-    ./UE4Editor "~/Documents/Unreal Projects/MyProject/MyProject.uproject"
+    ./UE5Editor "~/Documents/Unreal Projects/MyProject/MyProject.uproject"
     
 You can also append -game if you want to run the project as a game (you can also do that from the running editor).
 
diff --git a/Engine/Build/BatchFiles/Mac/SyncBuildRunLatest.command b/Engine/Build/BatchFiles/Mac/SyncBuildRunLatest.command
--- a/Engine/Build/BatchFiles/Mac/SyncBuildRunLatest.command
+++ b/Engine/Build/BatchFiles/Mac/SyncBuildRunLatest.command
@@ -36,13 +36,13 @@
 project=$1
 changelist=$2
 
-if pgrep -x "UE4Editor" > /dev/null
+if pgrep -x "UE5Editor" > /dev/null
 then
-    echo "UE4Editor already running, please close all UE4Editor processes..."
-    #killall -9 "UE4Editor"
+    echo "UE5Editor already running, please close all UE5Editor processes..."
+    #killall -9 "UE5Editor"
     exit 1
 else
-    echo "No UE4Editor processes running."
+    echo "No UE5Editor processes running."
 fi
 
 if [ ! -f ../../../Build/BatchFiles/Mac/SyncBuildRunLatest.command ]; then
@@ -135,13 +135,13 @@ fi
 echo changelist="$changelist"
 echo $latest_change_string
 
-if pgrep -x "UE4Editor" > /dev/null
+if pgrep -x "UE5Editor" > /dev/null
 then
-    echo "UE4Editor already running, please close all UE4Editor processes..."
+    echo "UE5Editor already running, please close all UE5Editor processes..."
     exit 1
-    #killall -9 "UE4Editor"
+    #killall -9 "UE5Editor"
 else
-    echo "No UE4Editor processes running."
+    echo "No UE5Editor processes running."
 fi
 
 echo $ue4_root_directory_local
@@ -157,6 +157,6 @@ else
   exit 1
 fi
 
-open -a "$ue4_root_directory_local"Engine/Binaries/Mac/UE4Editor.app --args "$project"
+open -a "$ue4_root_directory_local"Engine/Binaries/Mac/UE5Editor.app --args "$project"
 
 exit
\ No newline at end of file
diff --git a/Engine/Build/Graph/Examples/BuildEditorAndTools.xml b/Engine/Build/Graph/Examples/BuildEditorAndTools.xml
--- a/Engine/Build/Graph/Examples/BuildEditorAndTools.xml
+++ b/Engine/Build/Graph/Examples/BuildEditorAndTools.xml
@@ -22,7 +22,7 @@
 			RunUAT.bat BuildGraph -Script=Engine/Build/Graph/Examples/BuildEditorAndTools.xml -Target="Submit To Perforce For UGS" -P4 -Submit
 			
 		Other options:
-			-set:EditorTarget=<TargetName>		Specifies the editor target to be built (defaults to UE4Editor)
+			-set:EditorTarget=<TargetName>		Specifies the editor target to be built (defaults to UE5Editor)
 			-set:GameTargets=<TargetName>		Specifies the game targets to be built (defaults to none)
 			-set:TargetPlatforms=<A>;<B>;<C>... Specifies the target platforms to build the game for (defaults to Win64)
 			-set:OutputDir=<Path>				Specifies the output directory for the compiled binaries
@@ -31,7 +31,7 @@
 			-set:ArchiveStream=<Path>			Stream that will contain zip file of the compiled binaries, for use by UnrealGameSync.
 	-->
 
-	<Option Name="EditorTarget" Restrict="[^ ]+" DefaultValue="UE4Editor" Description="Name of the editor target to be built"/>
+	<Option Name="EditorTarget" Restrict="[^ ]+" DefaultValue="UE5Editor" Description="Name of the editor target to be built"/>
 	<Option Name="GameTargets" Restrict="[^ ]*" DefaultValue="" Description="Names of game targets to build separated by semicolons, eg. UE4Client;UE4Server"/>
 	<Option Name="TargetPlatforms" Restrict="[^ ]*" DefaultValue="Win64" Description="List of the target platforms to build for, separated by semicolons, eg. Win64;Android"/>
 	<Option Name="OutputDir" DefaultValue="$(RootDir)\LocalBuilds\$(EditorTarget)Binaries" Description ="Output directory for compiled binaries"/>
diff --git a/Engine/Build/InstalledEngineBuild.xml b/Engine/Build/InstalledEngineBuild.xml
--- a/Engine/Build/InstalledEngineBuild.xml
+++ b/Engine/Build/InstalledEngineBuild.xml
@@ -150,93 +150,93 @@
 		<Node Name="Compile UnrealHeaderTool Win64" Requires="Update Version Files">
 			<Compile Target="UnrealHeaderTool" Platform="Win64" Configuration="Development" Arguments="-precompile -allmodules"/>
 		</Node>
-		<Node Name="Compile UE4Editor Win64" Requires="Compile UnrealHeaderTool Win64" Produces="#UE4Editor Win64;#UE4Editor Win64 Unstripped;#UE4Editor Win64 Stripped;#UE4Editor Win64 Unsigned;#UE4Editor Win64 Signed;#UE4Editor Win64 WithExceptions">
-			<Compile Target="UE4Editor" Platform="Win64" Configuration="DebugGame" Tag="#UE4Editor Win64" Arguments="-precompile -allmodules "/>
-			<Compile Target="UE4Editor" Platform="Win64" Configuration="Development" Tag="#UE4Editor Win64" Arguments="-precompile -allmodules"/>
+		<Node Name="Compile UE5Editor Win64" Requires="Compile UnrealHeaderTool Win64" Produces="#UE5Editor Win64;#UE5Editor Win64 Unstripped;#UE5Editor Win64 Stripped;#UE5Editor Win64 Unsigned;#UE5Editor Win64 Signed;#UE5Editor Win64 WithExceptions">
+			<Compile Target="UE5Editor" Platform="Win64" Configuration="DebugGame" Tag="#UE5Editor Win64" Arguments="-precompile -allmodules "/>
+			<Compile Target="UE5Editor" Platform="Win64" Configuration="Development" Tag="#UE5Editor Win64" Arguments="-precompile -allmodules"/>
 
 			<Do If="$(EmbedSrcSrvInfo)">
 				<!-- Embed source info into the PDB files. Should be done from this machine to ensure that paths are correct. -->
 				<Log Message="Embedding source file information into PDB files..."/>
 				<Tag Files="Engine/Source/...;Engine/Plugins/..." Filter="*.c;*.h;*.cpp;*.hpp;*.inl" Except="Engine/Source/ThirdParty/..." With="#SourceFiles"/>
-				<SrcSrv BinaryFiles="#UE4Editor Win64" SourceFiles="#SourceFiles" Branch="$(Branch)" Change="$(Change)"/>
+				<SrcSrv BinaryFiles="#UE5Editor Win64" SourceFiles="#SourceFiles" Branch="$(Branch)" Change="$(Change)"/>
 			</Do>
 
-			<Tag Files="#UE4Editor Win64" Filter="$(Win64StripFilter)" Except="$(Win64StripExceptions);$(PluginsExceptions)" With="#UE4Editor Win64 Unstripped"/>
+			<Tag Files="#UE5Editor Win64" Filter="$(Win64StripFilter)" Except="$(Win64StripExceptions);$(PluginsExceptions)" With="#UE5Editor Win64 Unstripped"/>
 
 			<Switch>
 				<Case If="'$(WithFullDebugInfo)' == true">
 					<!-- Copy files for local debug builds -->
-					<Copy Files="#UE4Editor Win64 Unstripped" From="$(RootDir)" To="$(SavedOutput)"/>
+					<Copy Files="#UE5Editor Win64 Unstripped" From="$(RootDir)" To="$(SavedOutput)"/>
 				</Case>
 				<Default>
 					<!-- Otherwise, strip the files for distribution -->
-					<Strip BaseDir="$(RootDir)" OutputDir="$(SavedOutput)" Platform="Win64" Files="#UE4Editor Win64 Unstripped" Tag="#UE4Editor Win64 Stripped"/>
+					<Strip BaseDir="$(RootDir)" OutputDir="$(SavedOutput)" Platform="Win64" Files="#UE5Editor Win64 Unstripped" Tag="#UE5Editor Win64 Stripped"/>
 				</Default>
 			</Switch>
 
 			<Do If="'$(SignExecutables)' == true">
-				<Tag Files="#UE4Editor Win64" Filter="$(WindowsSignFilter)" Except="$(PluginsExceptions)" With="#UE4Editor Win64 Unsigned"/>
-				<Copy Files="#UE4Editor Win64 Unsigned" From="$(RootDir)" To="$(SavedOutput)" Tag="#Editor_ToSign_Win64"/>
-				<Sign Files="#Editor_ToSign_Win64" Tag="#UE4Editor Win64 Signed"/>
+				<Tag Files="#UE5Editor Win64" Filter="$(WindowsSignFilter)" Except="$(PluginsExceptions)" With="#UE5Editor Win64 Unsigned"/>
+				<Copy Files="#UE5Editor Win64 Unsigned" From="$(RootDir)" To="$(SavedOutput)" Tag="#Editor_ToSign_Win64"/>
+				<Sign Files="#Editor_ToSign_Win64" Tag="#UE5Editor Win64 Signed"/>
 			</Do>
-			<Tag Files="#UE4Editor Win64" Except="$(PluginsExceptions)" With="#UE4Editor Win64 WithExceptions"/>
+			<Tag Files="#UE5Editor Win64" Except="$(PluginsExceptions)" With="#UE5Editor Win64 WithExceptions"/>
 
 		</Node>
 	</Agent>
-	<Aggregate Name="Win64 Editor" Label="Editors/Win64" Requires="Compile UE4Editor Win64" Include="Editor Win64"/>
+	<Aggregate Name="Win64 Editor" Label="Editors/Win64" Requires="Compile UE5Editor Win64" Include="Editor Win64"/>
 
 	<!-- Base editor and tools for compile on Mac -->
 	<Agent Name="Editor Mac" Type="Mac_Licensee">
 		<Node Name="Compile UnrealHeaderTool Mac" Requires="Update Version Files">
 			<Compile Target="UnrealHeaderTool" Platform="Mac" Configuration="Development" Arguments="-precompile -allmodules"/>
 		</Node>
-		<Node Name="Compile UE4Editor Mac" Requires="Compile UnrealHeaderTool Mac" Produces="#UE4Editor Mac;#UE4Editor Mac Unstripped;#UE4Editor Mac Stripped;#UE4Editor Mac Unsigned;#UE4Editor Mac Signed;#UE4Editor Mac WithExceptions">
-			<Compile Target="UE4Editor" Platform="Mac" Configuration="DebugGame" Tag="#UE4Editor Mac" Arguments="-precompile -allmodules"/>
-			<Compile Target="UE4Editor" Platform="Mac" Configuration="Development" Tag="#UE4Editor Mac" Arguments="-precompile -allmodules"/>
+		<Node Name="Compile UE5Editor Mac" Requires="Compile UnrealHeaderTool Mac" Produces="#UE5Editor Mac;#UE5Editor Mac Unstripped;#UE5Editor Mac Stripped;#UE5Editor Mac Unsigned;#UE5Editor Mac Signed;#UE5Editor Mac WithExceptions">
+			<Compile Target="UE5Editor" Platform="Mac" Configuration="DebugGame" Tag="#UE5Editor Mac" Arguments="-precompile -allmodules"/>
+			<Compile Target="UE5Editor" Platform="Mac" Configuration="Development" Tag="#UE5Editor Mac" Arguments="-precompile -allmodules"/>
 
-			<Tag Files="#UE4Editor Mac" Filter="$(MacStripFilter)" Except="$(PluginsExceptions)" With="#UE4Editor Mac Unstripped"/>
+			<Tag Files="#UE5Editor Mac" Filter="$(MacStripFilter)" Except="$(PluginsExceptions)" With="#UE5Editor Mac Unstripped"/>
 
 			<Switch>
 				<Case If="'$(WithFullDebugInfo)' == true">
 					<!-- Copy files for local debug builds -->
-					<Copy Files="#UE4Editor Mac Unstripped" From="$(RootDir)" To="$(SavedOutput)"/>
+					<Copy Files="#UE5Editor Mac Unstripped" From="$(RootDir)" To="$(SavedOutput)"/>
 				</Case>
 				<Default>
 					<!-- Otherwise, strip the files for distribution -->
-					<Strip BaseDir="$(RootDir)" OutputDir="$(SavedOutput)" Platform="Mac" Files="#UE4Editor Mac Unstripped" Tag="#UE4Editor Mac Stripped"/>
+					<Strip BaseDir="$(RootDir)" OutputDir="$(SavedOutput)" Platform="Mac" Files="#UE5Editor Mac Unstripped" Tag="#UE5Editor Mac Stripped"/>
 				</Default>
 			</Switch>
 
 			<Do If="'$(SignExecutables)' == true">
-				<Tag Files="#UE4Editor Mac" Filter="$(MacSignFilter)" Except="$(PluginsExceptions)" With="#UE4Editor Mac Unsigned"/>
-				<Copy Files="#UE4Editor Mac Unsigned" From="$(RootDir)" To="$(SavedOutput)" Tag="#Editor_ToSign_Mac"/>
-				<Sign Files="#Editor_ToSign_Mac" Tag="#UE4Editor Mac Signed"/>
+				<Tag Files="#UE5Editor Mac" Filter="$(MacSignFilter)" Except="$(PluginsExceptions)" With="#UE5Editor Mac Unsigned"/>
+				<Copy Files="#UE5Editor Mac Unsigned" From="$(RootDir)" To="$(SavedOutput)" Tag="#Editor_ToSign_Mac"/>
+				<Sign Files="#Editor_ToSign_Mac" Tag="#UE5Editor Mac Signed"/>
 			</Do>
 
-			<Tag Files="#UE4Editor Mac" Except="$(PluginsExceptions)" With="#UE4Editor Mac WithExceptions"/>
+			<Tag Files="#UE5Editor Mac" Except="$(PluginsExceptions)" With="#UE5Editor Mac WithExceptions"/>
 		</Node>
 	</Agent>
-	<Aggregate Name="Mac Editor" Label="Editors/Mac" Requires="Compile UE4Editor Mac" Include="Editor Mac"/>
+	<Aggregate Name="Mac Editor" Label="Editors/Mac" Requires="Compile UE5Editor Mac" Include="Editor Mac"/>
 
 	<!-- Base editor and tools for compile on Linux -->
 	<Agent Name="Editor Linux" Type="Win64_Licensee">
 		<Node Name="Compile UnrealHeaderTool Linux" Requires="Update Version Files">
 			<Compile Target="UnrealHeaderTool" Platform="Linux" Configuration="Development" Arguments="-precompile -allmodules"/>
 		</Node>
-		<Node Name="Compile UE4Editor Linux" Requires="Compile UnrealHeaderTool $(HostPlatform)" Produces="#UE4Editor Linux;#UE4Editor Linux Unstripped;#UE4Editor Linux Stripped">
-			<Compile Target="UE4Editor" Platform="Linux" Configuration="DebugGame" Tag="#UE4Editor Linux" Arguments="-precompile -allmodules"/>
-			<Compile Target="UE4Editor" Platform="Linux" Configuration="Development" Tag="#UE4Editor Linux" Arguments="-precompile -allmodules"/>
+		<Node Name="Compile UE5Editor Linux" Requires="Compile UnrealHeaderTool $(HostPlatform)" Produces="#UE5Editor Linux;#UE5Editor Linux Unstripped;#UE5Editor Linux Stripped">
+			<Compile Target="UE5Editor" Platform="Linux" Configuration="DebugGame" Tag="#UE5Editor Linux" Arguments="-precompile -allmodules"/>
+			<Compile Target="UE5Editor" Platform="Linux" Configuration="Development" Tag="#UE5Editor Linux" Arguments="-precompile -allmodules"/>
 
-			<Tag Files="#UE4Editor Linux" Filter="$(LinuxStripFilter)" With="#UE4Editor Linux Unstripped"/>
+			<Tag Files="#UE5Editor Linux" Filter="$(LinuxStripFilter)" With="#UE5Editor Linux Unstripped"/>
 
 			<Switch>
 				<Case If="'$(WithFullDebugInfo)' == true">
 					<!-- Copy files for local debug builds -->
-					<Copy Files="#UE4Editor Linux Unstripped" From="$(RootDir)" To="$(SavedOutput)" Tag="#UE4Editor Linux Stripped"/>
+					<Copy Files="#UE5Editor Linux Unstripped" From="$(RootDir)" To="$(SavedOutput)" Tag="#UE5Editor Linux Stripped"/>
 				</Case>
 				<Default>
 					<!-- Otherwise, strip the files for distribution -->
-					<Strip BaseDir="$(RootDir)" OutputDir="$(SavedOutput)" Platform="Linux" Files="#UE4Editor Linux Unstripped" Tag="#UE4Editor Linux Stripped"/>
+					<Strip BaseDir="$(RootDir)" OutputDir="$(SavedOutput)" Platform="Linux" Files="#UE5Editor Linux Unstripped" Tag="#UE5Editor Linux Stripped"/>
 				</Default>
 			</Switch>
 		</Node>
@@ -608,7 +608,7 @@
 	</Macro>
 
 	<Agent Name="DDC Group Win64" Type="Win64_Licensee">
-		<Node Name="Build DDC Win64" Requires="#UE4Editor Win64;Build Tools Win64;Build Tools CS" Produces="#CompressedDDCWin64">
+		<Node Name="Build DDC Win64" Requires="#UE5Editor Win64;Build Tools Win64;Build Tools CS" Produces="#CompressedDDCWin64">
 			<!-- Get our temp dir -->
 			<Property Name="DDCDir" Value="$(RootDir)\LocalBuilds\InstalledDDC"/>
 			<Property Name="SavedDir" Value="$(SavedOutput)\Installed\Win64"/>
@@ -616,7 +616,7 @@
 			<Delete Files="$(DDCDir)\..."/>
 			<!-- Build up a list of files needed to build DDC -->
 			<CsCompile Project="Engine/Source/Programs/UnrealBuildTool/UnrealBuildTool.csproj" Configuration="Development" Platform="AnyCPU" Tag="#ToCopy" TagReferences="#ToCopy" EnumerateOnly="true"/>
-			<Tag Files="#Compile UnrealHeaderTool Win64;#UE4Editor Win64;#Build Tools Win64" With="#ToCopy"/>
+			<Tag Files="#Compile UnrealHeaderTool Win64;#UE5Editor Win64;#Build Tools Win64" With="#ToCopy"/>
 			<Tag Files="#ToCopy" Filter="*.target" With="#TargetReceipts"/>
 			<TagReceipt Files="#TargetReceipts" RuntimeDependencies="true" With="#ToCopy"/>
 			<Tag Files="Engine\Binaries\DotNET\Ionic.Zip.Reduced.dll" With="#ToCopy"/>
@@ -646,14 +646,14 @@
 	</Do>
 
 	<Agent Name="DDC Group Mac" Type="Mac">
-		<Node Name="Build DDC Mac" Requires="#UE4Editor Mac;Build Tools Mac;#Build Tools CS Binaries" Produces="#CompressedDDCMac">
+		<Node Name="Build DDC Mac" Requires="#UE5Editor Mac;Build Tools Mac;#Build Tools CS Binaries" Produces="#CompressedDDCMac">
 			<!-- Get our temp dir -->
 			<Property Name="DDCDir" Value="$(RootDir)/LocalBuilds/InstalledDDC"/>
 			<Property Name="SavedDir" Value="$(SavedOutput)/Installed/Mac"/>
 			<!-- Delete all the old files -->
 			<Delete Files="$(DDCDir)/..."/>
 			<!-- Build up a list of files needed to build DDC -->
-			<Tag Files="#Compile UnrealHeaderTool Mac;#UE4Editor Mac;#Build Tools Mac" With="#ToCopy"/>
+			<Tag Files="#Compile UnrealHeaderTool Mac;#UE5Editor Mac;#Build Tools Mac" With="#ToCopy"/>
 			<Tag Files="#ToCopy" Filter="*.target" With="#TargetReceipts"/>
 			<TagReceipt Files="#TargetReceipts" RuntimeDependencies="true" With="#ToCopy"/>
 			<Tag Files="Engine/Binaries/DotNET/Ionic.Zip.Reduced.dll" With="#ToCopy"/>
@@ -681,8 +681,8 @@
 
 	<Agent Name="DDC Group Linux" Type="Linux;Win64">
 		<!-- Additionally require the Win64 editor binaries if building on win64 -->
-		<Property Name="DDCLinuxRequires" Value="#UE4Editor Linux;Build Tools Linux;#Build Tools CS Binaries"/>
-		<Property Name="DDCLinuxRequires" Value="$(DDCLinuxRequires);#UE4Editor Win64;#Build Tools Win64" If="'$(HostPlatform)' == 'Win64'"/>
+		<Property Name="DDCLinuxRequires" Value="#UE5Editor Linux;Build Tools Linux;#Build Tools CS Binaries"/>
+		<Property Name="DDCLinuxRequires" Value="$(DDCLinuxRequires);#UE5Editor Win64;#Build Tools Win64" If="'$(HostPlatform)' == 'Win64'"/>
 		<Node Name="Build DDC Linux" Requires="$(DDCLinuxRequires)" Produces="#CompressedDDCLinux">
 			<!-- Get our temp dir -->
 			<Property Name="DDCDir" Value="$(RootDir)/LocalBuilds/InstalledDDC"/>
@@ -691,8 +691,8 @@
 			<Delete Files="$(DDCDir)/..."/>
 			<!-- Build up a list of files needed to build DDC -->
 			<CsCompile Project="Engine/Source/Programs/UnrealBuildTool/UnrealBuildTool.csproj" Configuration="Development" Platform="AnyCPU" Tag="#ToCopy" TagReferences="#ToCopy" EnumerateOnly="true"/>
-			<Tag Files="#Compile UnrealHeaderTool $(HostPlatform);#UE4Editor Linux;#Build Tools Linux" With="#ToCopy"/>
-			<Tag Files="#UE4Editor Win64;#Build Tools Win64" With="#ToCopy" If="'$(HostPlatform)' == 'Win64'"/>
+			<Tag Files="#Compile UnrealHeaderTool $(HostPlatform);#UE5Editor Linux;#Build Tools Linux" With="#ToCopy"/>
+			<Tag Files="#UE5Editor Win64;#Build Tools Win64" With="#ToCopy" If="'$(HostPlatform)' == 'Win64'"/>
 			<Tag Files="#ToCopy" Filter="*.target" With="#TargetReceipts"/>
 			<TagReceipt Files="#TargetReceipts" RuntimeDependencies="true" With="#ToCopy"/>
 			<Tag Files="Engine/Content/..." Except="*.psd;*.pdn;*.fbx;*.po" With="#ToCopy"/>
@@ -722,7 +722,7 @@
 			<Do If="$(CompileDatasmithPlugins) == 'True'">
 				<Property Name="DatasmithPlugins" Value="DatasmithCADImporter;DatasmithIFCImporter;DatasmithC4DImporter;AxFImporter;StaticMeshEditorExtension;MDLImporter"/>
 				<ForEach Name="Plugin" Values="$(DatasmithPlugins)">
-					<Compile Target="UE4Editor" Configuration="Development" Platform="Win64" Arguments="-Plugin=&quot;$(RootDir)/Engine/Plugins/Enterprise/$(Plugin)/$(Plugin).uplugin&quot;"/>
+					<Compile Target="UE5Editor" Configuration="Development" Platform="Win64" Arguments="-Plugin=&quot;$(RootDir)/Engine/Plugins/Enterprise/$(Plugin)/$(Plugin).uplugin&quot;"/>
 					<Tag Files="$(RootDir)/Engine/Plugins/Enterprise/$(Plugin)/Binaries/..." With="#DatasmithPluginBinaries Win64"/>
 				</ForEach>
 			</Do>
@@ -734,7 +734,7 @@
 			<Do If="$(CompileDatasmithPlugins) == 'True'">
 				<Property Name="DatasmithPlugins" Value="DatasmithCADImporter;DatasmithIFCImporter;DatasmithC4DImporter;AxFImporter;StaticMeshEditorExtension;MDLImporter"/>
 				<ForEach Name="Plugin" Values="$(DatasmithPlugins)">
-					<Compile Target="UE4Editor" Configuration="Development" Platform="Mac" Arguments="-Plugin=&quot;$(RootDir)/Engine/Plugins/Enterprise/$(Plugin)/$(Plugin).uplugin&quot;"/>
+					<Compile Target="UE5Editor" Configuration="Development" Platform="Mac" Arguments="-Plugin=&quot;$(RootDir)/Engine/Plugins/Enterprise/$(Plugin)/$(Plugin).uplugin&quot;"/>
 					<Tag Files="$(RootDir)/Engine/Plugins/Enterprise/$(Plugin)/Binaries/..." With="#DatasmithPluginBinaries Mac"/>
 				</ForEach>
 			</Do>
@@ -743,7 +743,7 @@
 
 	<Agent Name="Installed Build Group Win64" Type="Win64_Licensee">
 		<!-- Build up a list of requirements as some can be optional -->
-		<Property Name="InstalledRequirements" Value="Compile UnrealHeaderTool Win64;#UE4Editor Win64 WithExceptions;#UE4Editor Win64 Unstripped;#UE4Editor Win64 Stripped;#UE4Editor Win64 Unsigned;#UE4Editor Win64 Signed;#Build Tools Win64;#Build Tools CS Binaries;#Rules Assemblies"/>
+		<Property Name="InstalledRequirements" Value="Compile UnrealHeaderTool Win64;#UE5Editor Win64 WithExceptions;#UE5Editor Win64 Unstripped;#UE5Editor Win64 Stripped;#UE5Editor Win64 Unsigned;#UE5Editor Win64 Signed;#Build Tools Win64;#Build Tools CS Binaries;#Rules Assemblies"/>
 		<!-- Optional Target Platform requirements -->
 		<ForEach Name="Platform" Values="$(ExtensionPlatforms)">
 			<Expand Name="$(Platform)_AddInstalledRequirements"/>
@@ -776,8 +776,8 @@
 
 			<!-- Tag files that always need to be part of an installed build -->
 			<Tag Files="#Compile UnrealHeaderTool Win64" With="#Installed Win64"/>
-			<Tag Files="#UE4Editor Win64 WithExceptions" Except="#UE4Editor Win64 Unstripped;#UE4Editor Win64 Unsigned" With="#Installed Win64"/>
-			<Tag Files="#UE4Editor Win64 Stripped;#UE4Editor Win64 Signed" With="#Saved Output"/>
+			<Tag Files="#UE5Editor Win64 WithExceptions" Except="#UE5Editor Win64 Unstripped;#UE5Editor Win64 Unsigned" With="#Installed Win64"/>
+			<Tag Files="#UE5Editor Win64 Stripped;#UE5Editor Win64 Signed" With="#Saved Output"/>
 			<Tag Files="#Build Tools Win64;#Make Feature Packs;#Rules Assemblies" With="#Installed Win64"/>
 			<Tag Files="#DatasmithPluginBinaries Win64" With="#Installed Win64" If="$(CompileDatasmithPlugins)"/>
 
@@ -904,7 +904,7 @@
 
 	<Agent Name="Installed Build Group Mac" Type="Mac_Licensee">
 		<!-- Build up a list of requirements as some can be optional -->
-		<Property Name="FilterRequirements" Value="#Compile UnrealHeaderTool Mac;#UE4Editor Mac WithExceptions;#UE4Editor Mac Unstripped;#UE4Editor Mac Stripped;#UE4Editor Mac Unsigned;#UE4Editor Mac Signed;#Build Tools Mac;#Build Tools CS Binaries;#Rules Assemblies"/>
+		<Property Name="FilterRequirements" Value="#Compile UnrealHeaderTool Mac;#UE5Editor Mac WithExceptions;#UE5Editor Mac Unstripped;#UE5Editor Mac Stripped;#UE5Editor Mac Unsigned;#UE5Editor Mac Signed;#Build Tools Mac;#Build Tools CS Binaries;#Rules Assemblies"/>
 
 		<!-- Optional Target Platform requirements -->
 		<Property Name="FilterRequirements" Value="$(FilterRequirements);#UE4Game Mac;#UE4Game Mac Unstripped;#UE4Game Mac Stripped;#UE4Game Mac Unsigned;#UE4Game Mac Signed" If="'$(WithMac)' == true"/>
@@ -932,8 +932,8 @@
 
 			<!-- Tag files that always need to be part of an installed build -->
 			<Tag Files="#Compile UnrealHeaderTool Mac" With="#Installed Mac"/>
-			<Tag Files="#UE4Editor Mac WithExceptions" Except="#UE4Editor Mac Unstripped;#UE4Editor Mac Unsigned" With="#Installed Mac"/>
-			<Tag Files="#UE4Editor Mac Stripped;#UE4Editor Mac Signed" With="#Saved Output"/>
+			<Tag Files="#UE5Editor Mac WithExceptions" Except="#UE5Editor Mac Unstripped;#UE5Editor Mac Unsigned" With="#Installed Mac"/>
+			<Tag Files="#UE5Editor Mac Stripped;#UE5Editor Mac Signed" With="#Saved Output"/>
 			<Tag Files="#Build Tools Mac;#Make Feature Packs;#Rules Assemblies" With="#Installed Mac"/>
 			<Tag Files="#DatasmithPluginBinaries Mac" With="#Installed Mac" If="$(CompileDatasmithPlugins)"/>
 
@@ -1034,7 +1034,7 @@
 
 	<Agent Name="Installed Build Group Linux" Type="Win64_Licensee">
 		<!-- Build up a list of requirements as some can be optional -->
-		<Property Name="InstalledRequirements" Value="#Compile UnrealHeaderTool Linux;#UE4Editor Linux;#UE4Editor Linux Unstripped;#UE4Editor Linux Stripped;#Build Tools Linux;#Build Tools CS Binaries;#Rules Assemblies"/>
+		<Property Name="InstalledRequirements" Value="#Compile UnrealHeaderTool Linux;#UE5Editor Linux;#UE5Editor Linux Unstripped;#UE5Editor Linux Stripped;#Build Tools Linux;#Build Tools CS Binaries;#Rules Assemblies"/>
 
 		<!-- Optional Target Platform requirements -->
 		<Property Name="InstalledRequirements" Value="$(InstalledRequirements);#UE4Game Linux;#UE4Game Linux Unstripped;#UE4Game Linux Stripped" If="'$(WithLinux)' == true"/>
@@ -1059,8 +1059,8 @@
 
 			<!-- Tag files that always need to be part of an installed build -->
 			<Tag Files="#Compile UnrealHeaderTool Linux" With="#Installed Linux"/>
-			<Tag Files="#UE4Editor Linux" Except="#UE4Editor Linux Unstripped" With="#Installed Linux"/>
-			<Tag Files="#UE4Editor Linux Stripped" With="#Saved Output"/>
+			<Tag Files="#UE5Editor Linux" Except="#UE5Editor Linux Unstripped" With="#Installed Linux"/>
+			<Tag Files="#UE5Editor Linux Stripped" With="#Saved Output"/>
 			<Tag Files="#Build Tools Linux;#Make Feature Packs;#Rules Assemblies" With="#Installed Linux"/>
 
 			<!-- Tag enabled target platforms -->
diff --git a/Engine/Build/InstalledEngineFilters.xml b/Engine/Build/InstalledEngineFilters.xml
--- a/Engine/Build/InstalledEngineFilters.xml
+++ b/Engine/Build/InstalledEngineFilters.xml
@@ -15,7 +15,7 @@
 
 	<!-- List of file types to be stripped and signed for different platforms -->
 	<Property Name="Win64StripFilter" Value="*.pdb"/>
-	<Property Name="Win64StripExceptions" Value="Engine\Binaries\Win64\UE4Editor*.pdb;Engine\Plugins\...\Binaries\Win64\UE4Editor*.pdb"/>
+	<Property Name="Win64StripExceptions" Value="Engine\Binaries\Win64\UE5Editor*.pdb;Engine\Plugins\...\Binaries\Win64\UE5Editor*.pdb"/>
 	<Property Name="MacStripFilter" Value="*.a;"/>
 	<Property Name="AndroidStripFilter" Value="*.a;*.so"/>
 	<Property Name="IOSStripFilter" Value="*.a;"/>
@@ -108,7 +108,7 @@
 
 		<!-- Source code -->
 		Engine/Source/UE4Game.Target.cs
-		Engine/Source/UE4Editor.Target.cs
+		Engine/Source/UE5Editor.Target.cs
 
 		<!-- Starter content -->
 		Samples/StarterContent/Content/...
diff --git a/Engine/Plugins/Enterprise/DataprepEditor/Source/DataprepEditor/Private/DataprepEditor.cpp b/Engine/Plugins/Enterprise/DataprepEditor/Source/DataprepEditor/Private/DataprepEditor.cpp
--- a/Engine/Plugins/Enterprise/DataprepEditor/Source/DataprepEditor/Private/DataprepEditor.cpp
+++ b/Engine/Plugins/Enterprise/DataprepEditor/Source/DataprepEditor/Private/DataprepEditor.cpp
@@ -315,7 +315,7 @@ void FDataprepEditor::CleanUpTemporaryDirectories()
 					if(!bDeleteDirectories)
 					{
 						const FString ApplicationName = FPlatformProcess::GetApplicationName( ProcessID );
-						bDeleteDirectories = !ApplicationName.StartsWith(TEXT("UE4Editor"));
+						bDeleteDirectories = !ApplicationName.StartsWith(TEXT("UE5Editor"));
 					}
 
 					if(bDeleteDirectories)
diff --git a/Engine/Plugins/NetcodeUnitTest/NetcodeUnitTest/Source/NetcodeUnitTest/Classes/UnitTestCommandlet.h b/Engine/Plugins/NetcodeUnitTest/NetcodeUnitTest/Source/NetcodeUnitTest/Classes/UnitTestCommandlet.h
--- a/Engine/Plugins/NetcodeUnitTest/NetcodeUnitTest/Source/NetcodeUnitTest/Classes/UnitTestCommandlet.h
+++ b/Engine/Plugins/NetcodeUnitTest/NetcodeUnitTest/Source/NetcodeUnitTest/Classes/UnitTestCommandlet.h
@@ -13,7 +13,7 @@
  * Uses a heavily stripped down game client, with allowances for Slate, and some hacks for enabling proper netcode usage.
  *
  * Usage:
- *	"UE4Editor.exe shootergame -run=NetcodeUnitTest.UnitTestCommandlet"
+ *	"UE5Editor.exe shootergame -run=NetcodeUnitTest.UnitTestCommandlet"
  *
  * Parameters:
  *	-UnitTest=UnitTestName
diff --git a/Engine/Plugins/NetcodeUnitTest/NetcodeUnitTest/Source/NetcodeUnitTest/Private/ProcessUnitTest.cpp b/Engine/Plugins/NetcodeUnitTest/NetcodeUnitTest/Source/NetcodeUnitTest/Private/ProcessUnitTest.cpp
--- a/Engine/Plugins/NetcodeUnitTest/NetcodeUnitTest/Source/NetcodeUnitTest/Private/ProcessUnitTest.cpp
+++ b/Engine/Plugins/NetcodeUnitTest/NetcodeUnitTest/Source/NetcodeUnitTest/Private/ProcessUnitTest.cpp
@@ -748,7 +748,7 @@ bool UProcessUnitTest::IsBlockingProcessPresent(bool bLogIfFound/*=false*/)
 					TEXT("Recommend generating shaders separately before running unit tests, to avoid bugs and unit test failure."));
 
 			UNIT_LOG(ELogType::StatusImportant,
-					TEXT("One way to do this (takes a long time), is through: UE4Editor-Cmd YourGame -run=DerivedDataCache -fill"));
+					TEXT("One way to do this (takes a long time), is through: UE5Editor-Cmd YourGame -run=DerivedDataCache -fill"));
 		}
 	}
 
diff --git a/Engine/Programs/UnrealGameSync/UnrealGameSync.ini b/Engine/Programs/UnrealGameSync/UnrealGameSync.ini
--- a/Engine/Programs/UnrealGameSync/UnrealGameSync.ini
+++ b/Engine/Programs/UnrealGameSync/UnrealGameSync.ini
@@ -16,6 +16,8 @@ DebugGameHasSeparateExecutable=1
 +ExcludeChanges=^Collection '.*' not modified
 
 [Default]
+EditorTarget=UE5Editor
+
 ; Set this to a network share which contains the SDK installers for your site
 SdkInstallerDir= 
 
diff --git a/Engine/Source/Developer/DesktopPlatform/Private/DesktopPlatformBase.cpp b/Engine/Source/Developer/DesktopPlatform/Private/DesktopPlatformBase.cpp
--- a/Engine/Source/Developer/DesktopPlatform/Private/DesktopPlatformBase.cpp
+++ b/Engine/Source/Developer/DesktopPlatform/Private/DesktopPlatformBase.cpp
@@ -565,10 +565,10 @@ bool FDesktopPlatformBase::CompileGameProject(const FString& RootDir, const FStr
 	// Build the argument list
 	FString Arguments = FString::Printf(TEXT("%s %s"), ModuleManager.GetUBTConfiguration(), FPlatformMisc::GetUBTPlatform());
 
-	// Append the project name if it's a foreign project. Otherwise compile UE4Editor.
+	// Append the project name if it's a foreign project. Otherwise compile UE5Editor.
 	if ( ProjectFileName.IsEmpty() )
 	{
-		Arguments = TEXT("UE4Editor ") + Arguments;
+		Arguments = TEXT("UE5Editor ") + Arguments;
 	}
 	else
 	{
diff --git a/Engine/Source/Developer/DesktopPlatform/Private/Mac/DesktopPlatformMac.cpp b/Engine/Source/Developer/DesktopPlatform/Private/Mac/DesktopPlatformMac.cpp
--- a/Engine/Source/Developer/DesktopPlatform/Private/Mac/DesktopPlatformMac.cpp
+++ b/Engine/Source/Developer/DesktopPlatform/Private/Mac/DesktopPlatformMac.cpp
@@ -538,7 +538,7 @@ void FDesktopPlatformMac::EnumerateEngineInstallations(TMap<FString, FString> &O
 			NSURL* AppURL = (NSURL*)CFArrayGetValueAtIndex(AllApps, Index);
 			NSBundle* AppBundle = [NSBundle bundleWithURL:AppURL];
 			FString EngineDir = FString([[AppBundle bundlePath] stringByDeletingLastPathComponent]);
-			if (([[AppBundle bundleIdentifier] isEqualToString:@"com.epicgames.UE4Editor"] || [[AppBundle bundleIdentifier] isEqualToString:@"com.epicgames.UE4EditorServices"])
+			if (([[AppBundle bundleIdentifier] isEqualToString:@"com.epicgames.UE5Editor"] || [[AppBundle bundleIdentifier] isEqualToString:@"com.epicgames.UE4EditorServices"])
 				&& EngineDir.RemoveFromEnd(TEXT("/Engine/Binaries/Mac")) && !EngineDir.Contains("Unreal Engine.app/Contents/") && !EngineDir.Contains("Epic Games Launcher.app/Contents/") && !EngineDir.Contains("[REDACTED_USER_HOME]/UnrealEngine/Launcher"))
 			{
 				FString EngineId;
diff --git a/Engine/Source/Developer/Linux/LinuxTargetPlatform/Private/LinuxTargetDevice.h b/Engine/Source/Developer/Linux/LinuxTargetPlatform/Private/LinuxTargetDevice.h
--- a/Engine/Source/Developer/Linux/LinuxTargetPlatform/Private/LinuxTargetDevice.h
+++ b/Engine/Source/Developer/Linux/LinuxTargetPlatform/Private/LinuxTargetDevice.h
@@ -161,7 +161,7 @@ public:
 		}
 		else if (TargetType == EBuildTargetType::Editor)
 		{
-			ExecutablePath /= TEXT("UE4Editor");
+			ExecutablePath /= TEXT("UE5Editor");
 		}
 
 		if (BuildConfiguration != EBuildConfiguration::Development)
diff --git a/Engine/Source/Developer/Windows/WindowsTargetPlatform/Private/LocalPcTargetDevice.h b/Engine/Source/Developer/Windows/WindowsTargetPlatform/Private/LocalPcTargetDevice.h
--- a/Engine/Source/Developer/Windows/WindowsTargetPlatform/Private/LocalPcTargetDevice.h
+++ b/Engine/Source/Developer/Windows/WindowsTargetPlatform/Private/LocalPcTargetDevice.h
@@ -253,7 +253,7 @@ public:
 		}
 		else if (TargetType == EBuildTargetType::Editor)
 		{
-			ExecutablePath /= TEXT("UE4Editor");
+			ExecutablePath /= TEXT("UE5Editor");
 		}
 
 		if (BuildConfiguration != EBuildConfiguration::Development)
diff --git a/Engine/Source/Editor/GameProjectGeneration/Private/GameProjectUtils.cpp b/Engine/Source/Editor/GameProjectGeneration/Private/GameProjectUtils.cpp
--- a/Engine/Source/Editor/GameProjectGeneration/Private/GameProjectUtils.cpp
+++ b/Engine/Source/Editor/GameProjectGeneration/Private/GameProjectUtils.cpp
@@ -4109,7 +4109,7 @@ GameProjectUtils::EAddCodeToProjectResult GameProjectUtils::AddCodeToProject_Int
 		const FString GameModuleName = FApp::GetProjectName();
 		check(ModuleInfo.ModuleName == GameModuleName);
 
-		// Because this project previously didn't have any code, the UBT target name will just be UE4Editor. Now that we've
+		// Because this project previously didn't have any code, the UBT target name will just be UE5Editor. Now that we've
 		// added some code, the target name will be changed to match the editor target for the new source. 
 		FString NewUBTTargetName = GameModuleName + TEXT("Editor");
 		FPlatformMisc::SetUBTTargetName(*NewUBTTargetName);
diff --git a/Engine/Source/Editor/UnrealEd/Private/Factories/PackFactory.cpp b/Engine/Source/Editor/UnrealEd/Private/Factories/PackFactory.cpp
--- a/Engine/Source/Editor/UnrealEd/Private/Factories/PackFactory.cpp
+++ b/Engine/Source/Editor/UnrealEd/Private/Factories/PackFactory.cpp
@@ -701,7 +701,7 @@ UObject* UPackFactory::FactoryCreateBinary
 					}
 					else
 					{
-						// We didn't previously have source, so the UBT target name will be UE4Editor, and attempts to recompile will end up building the wrong target. Now that we have source,
+						// We didn't previously have source, so the UBT target name will be UE5Editor, and attempts to recompile will end up building the wrong target. Now that we have source,
 						// we need to change the UBT target to be the newly created editor module
 						FPlatformMisc::SetUBTTargetName(*(FString(FApp::GetProjectName()) + TEXT("Editor")));
 
diff --git a/Engine/Source/Programs/AutomationTool/AutomationUtils/CommandletUtils.cs b/Engine/Source/Programs/AutomationTool/AutomationUtils/CommandletUtils.cs
--- a/Engine/Source/Programs/AutomationTool/AutomationUtils/CommandletUtils.cs
+++ b/Engine/Source/Programs/AutomationTool/AutomationUtils/CommandletUtils.cs
@@ -53,7 +53,7 @@ namespace AutomationTool
         /// <param name="CulturesToCook">List of culture names whose localized assets should be cooked, can be null (implying defaults should be used).</param>
 		/// <param name="TargetPlatform">Target platform.</param>
 		/// <param name="Parameters">List of additional parameters.</param>
-		public static void CookCommandlet(FileReference ProjectName, string UE4Exe = "UE4Editor-Cmd.exe", string[] Maps = null, string[] Dirs = null, string InternationalizationPreset = "", string[] CulturesToCook = null, string TargetPlatform = "Windows", string Parameters = "-Unversioned")
+		public static void CookCommandlet(FileReference ProjectName, string UE4Exe = "UE5Editor-Cmd.exe", string[] Maps = null, string[] Dirs = null, string InternationalizationPreset = "", string[] CulturesToCook = null, string TargetPlatform = "Windows", string Parameters = "-Unversioned")
 		{
             string CommandletArguments = "";
 
@@ -97,7 +97,7 @@ namespace AutomationTool
         /// <param name="Maps">List of maps to cook, can be null in which case -MapIniSection=AllMaps is used.</param>
         /// <param name="TargetPlatform">Target platform.</param>
         /// <param name="Parameters">List of additional parameters.</param>
-        public static void DDCCommandlet(FileReference ProjectName, string UE4Exe = "UE4Editor-Cmd.exe", string[] Maps = null, string TargetPlatform = "Windows", string Parameters = "")
+        public static void DDCCommandlet(FileReference ProjectName, string UE4Exe = "UE5Editor-Cmd.exe", string[] Maps = null, string TargetPlatform = "Windows", string Parameters = "")
         {
             string MapsToCook = "";
             if (!IsNullOrEmpty(Maps))
@@ -115,7 +115,7 @@ namespace AutomationTool
 		/// <param name="UE4Exe">The name of the UE4 Editor executable to use.</param>
 		/// <param name="Maps">List of maps to rebuild light maps for. Can be null in which case -MapIniSection=AllMaps is used.</param>
 		/// <param name="Parameters">List of additional parameters.</param>
-		public static void RebuildLightMapsCommandlet(FileReference ProjectName, string UE4Exe = "UE4Editor-Cmd.exe", string[] Maps = null, string Parameters = "")
+		public static void RebuildLightMapsCommandlet(FileReference ProjectName, string UE4Exe = "UE5Editor-Cmd.exe", string[] Maps = null, string Parameters = "")
 		{
 			string MapsToRebuildLighting = "";
 			if (!IsNullOrEmpty(Maps))
@@ -126,7 +126,7 @@ namespace AutomationTool
 			RunCommandlet(ProjectName, UE4Exe, "ResavePackages", String.Format("-buildtexturestreaming -buildlighting -MapsOnly -ProjectOnly -AllowCommandletRendering -SkipSkinVerify {0} {1}", MapsToRebuildLighting, Parameters));
 		}
 
-        public static void RebuildHLODCommandlet(FileReference ProjectName, string UE4Exe = "UE4Editor-Cmd.exe", string[] Maps = null, string Parameters = "")
+        public static void RebuildHLODCommandlet(FileReference ProjectName, string UE4Exe = "UE5Editor-Cmd.exe", string[] Maps = null, string Parameters = "")
         {
             string MapsToRebuildHLODs = "";
             if (!IsNullOrEmpty(Maps))
@@ -144,7 +144,7 @@ namespace AutomationTool
         /// <param name="UE4Exe">The name of the UE4 Editor executable to use.</param>
         /// <param name="Maps">List of maps to rebuild light maps for. Can be null in which case -MapIniSection=AllMaps is used.</param>
         /// <param name="Parameters">List of additional parameters.</param>
-        public static void ResavePackagesCommandlet(FileReference ProjectName, string UE4Exe = "UE4Editor-Cmd.exe", string[] Maps = null, string Parameters = "")
+        public static void ResavePackagesCommandlet(FileReference ProjectName, string UE4Exe = "UE5Editor-Cmd.exe", string[] Maps = null, string Parameters = "")
         {
             string MapsToRebuildLighting = "";
             if (!IsNullOrEmpty(Maps))
@@ -163,7 +163,7 @@ namespace AutomationTool
         /// <param name="Maps">List of maps to cook, can be null in which case -MapIniSection=AllMaps is used.</param>
         /// <param name="TargetPlatform">Target platform.</param>
         /// <param name="Parameters">List of additional parameters.</param>
-        public static List<FileReference> GenerateDistillFileSetsCommandlet(FileReference ProjectName, string ManifestFile, string UE4Exe = "UE4Editor-Cmd.exe", string[] Maps = null, string Parameters = "")
+        public static List<FileReference> GenerateDistillFileSetsCommandlet(FileReference ProjectName, string ManifestFile, string UE4Exe = "UE5Editor-Cmd.exe", string[] Maps = null, string Parameters = "")
         {
             string MapsToCook = "";
             if (!IsNullOrEmpty(Maps))
@@ -219,13 +219,13 @@ namespace AutomationTool
         /// <param name="ProjectName">Project name.</param>
         /// <param name="UE4Exe">The name of the UE4 Editor executable to use.</param>
         /// <param name="Parameters">List of additional parameters.</param>
-        public static void UpdateGameProjectCommandlet(FileReference ProjectName, string UE4Exe = "UE4Editor-Cmd.exe", string Parameters = "")
+        public static void UpdateGameProjectCommandlet(FileReference ProjectName, string UE4Exe = "UE5Editor-Cmd.exe", string Parameters = "")
         {
             RunCommandlet(ProjectName, UE4Exe, "UpdateGameProject", Parameters);
         }
 
 		/// <summary>
-		/// Runs a commandlet using Engine/Binaries/Win64/UE4Editor-Cmd.exe.
+		/// Runs a commandlet using Engine/Binaries/Win64/UE5Editor-Cmd.exe.
 		/// </summary>
 		/// <param name="ProjectFile">Project name.</param>
 		/// <param name="UE4Exe">The name of the UE4 Editor executable to use.</param>
@@ -238,7 +238,7 @@ namespace AutomationTool
 		}
 
 		/// <summary>
-		/// Runs a commandlet using Engine/Binaries/Win64/UE4Editor-Cmd.exe.
+		/// Runs a commandlet using Engine/Binaries/Win64/UE5Editor-Cmd.exe.
 		/// </summary>
 		/// <param name="ProjectFile">Project name.</param>
 		/// <param name="UE4Exe">The name of the UE4 Editor executable to use.</param>
@@ -247,7 +247,7 @@ namespace AutomationTool
 		/// <param name="DestLogFile">Log file after completion</param>
 		public static void RunCommandlet(FileReference ProjectName, string UE4Exe, string Commandlet, string Parameters, out string DestLogFile)
 		{
-			LogInformation("Running UE4Editor {0} for project {1}", Commandlet, ProjectName);
+			LogInformation("Running UE5Editor {0} for project {1}", Commandlet, ProjectName);
 
             var CWD = Path.GetDirectoryName(UE4Exe);
 
@@ -433,15 +433,15 @@ namespace AutomationTool
 		{
 			if (HostPlatform == UnrealBuildTool.UnrealTargetPlatform.Mac)
 			{
-				return CommandUtils.CombinePaths(BuildRoot, "Engine/Binaries/Mac/UE4Editor.app/Contents/MacOS/UE4Editor");
+				return CommandUtils.CombinePaths(BuildRoot, "Engine/Binaries/Mac/UE5Editor.app/Contents/MacOS/UE5Editor");
 			}
 			if (HostPlatform == UnrealBuildTool.UnrealTargetPlatform.Win64)
 			{
-				return CommandUtils.CombinePaths(BuildRoot, "Engine/Binaries/Win64/UE4Editor-Cmd.exe");
+				return CommandUtils.CombinePaths(BuildRoot, "Engine/Binaries/Win64/UE5Editor-Cmd.exe");
 			}
 			if (HostPlatform == UnrealBuildTool.UnrealTargetPlatform.Linux)
 			{
-				return CommandUtils.CombinePaths(BuildRoot, "Engine/Binaries/Linux/UE4Editor");
+				return CommandUtils.CombinePaths(BuildRoot, "Engine/Binaries/Linux/UE5Editor");
 			}
 			throw new AutomationException("EditorCommandlet is not supported for platform {0}", HostPlatform);
 		}
diff --git a/Engine/Source/Programs/AutomationTool/AutomationUtils/ProjectParams.cs b/Engine/Source/Programs/AutomationTool/AutomationUtils/ProjectParams.cs
--- a/Engine/Source/Programs/AutomationTool/AutomationUtils/ProjectParams.cs
+++ b/Engine/Source/Programs/AutomationTool/AutomationUtils/ProjectParams.cs
@@ -839,7 +839,7 @@ namespace AutomationTool
             this.RunAutomationTest = ParseParamValueIfNotSpecified(Command, RunAutomationTest, "RunAutomationTest");
             this.RunAutomationTests = this.RunAutomationTest != "" || GetParamValueIfNotSpecified(Command, RunAutomationTests, this.RunAutomationTests, "RunAutomationTests");
             this.SkipServer = GetParamValueIfNotSpecified(Command, SkipServer, this.SkipServer, "skipserver");
-			this.UE4Exe = ParseParamValueIfNotSpecified(Command, UE4Exe, "ue4exe", "UE4Editor-Cmd.exe");
+			this.UE4Exe = ParseParamValueIfNotSpecified(Command, UE4Exe, "ue4exe", "UE5Editor-Cmd.exe");
 			this.Unattended = GetParamValueIfNotSpecified(Command, Unattended, this.Unattended, "unattended");
 			this.DeviceUsername = ParseParamValueIfNotSpecified(Command, DeviceUsername, "deviceuser", String.Empty);
 			this.DevicePassword = ParseParamValueIfNotSpecified(Command, DevicePassword, "devicepass", String.Empty);
@@ -1336,7 +1336,7 @@ namespace AutomationTool
 		[Help("stagingdirectory=Path", "Directory to copy the builds to, i.e. -stagingdirectory=C:\\Stage")]
 		public string StageDirectoryParam;
         
-		[Help("ue4exe=ExecutableName", "Name of the UE4 Editor executable, i.e. -ue4exe=UE4Editor.exe")]
+		[Help("ue4exe=ExecutableName", "Name of the UE4 Editor executable, i.e. -ue4exe=UE5Editor.exe")]
 		public string UE4Exe;
 
 		/// <summary>
@@ -2052,7 +2052,7 @@ namespace AutomationTool
 			if (!bIsCodeBasedProject)
 			{
 				GameTarget = Client ? "UE4Client" : "UE4Game";
-				EditorTarget = "UE4Editor";
+				EditorTarget = "UE5Editor";
 				ServerTarget = "UE4Server";
 			}
 			else if (TargetNames.Count > 0)
@@ -2222,7 +2222,7 @@ namespace AutomationTool
 			{
 				if (Properties.bWasGenerated)
 				{
-					EditorTarget = "UE4Editor";
+					EditorTarget = "UE5Editor";
 				}
 				else
 				{
diff --git a/Engine/Source/Programs/AutomationTool/AutomationUtils/ProjectUtils.cs b/Engine/Source/Programs/AutomationTool/AutomationUtils/ProjectUtils.cs
--- a/Engine/Source/Programs/AutomationTool/AutomationUtils/ProjectUtils.cs
+++ b/Engine/Source/Programs/AutomationTool/AutomationUtils/ProjectUtils.cs
@@ -761,7 +761,7 @@ namespace AutomationTool
 			Type[] AllCompiledTypes = TargetsDLL.GetTypes();
 			foreach (Type TargetType in AllCompiledTypes)
 			{
-				// Find TargetRules but skip all "UE4Editor", "UE4Game" targets.
+				// Find TargetRules but skip all "UE5Editor", "UE4Game" targets.
 				if (typeof(TargetRules).IsAssignableFrom(TargetType) && !TargetType.IsAbstract)
 				{
 					string TargetName = GetTargetName(TargetType);
diff --git a/Engine/Source/Programs/AutomationTool/AutomationUtils/UE4Build.cs b/Engine/Source/Programs/AutomationTool/AutomationUtils/UE4Build.cs
--- a/Engine/Source/Programs/AutomationTool/AutomationUtils/UE4Build.cs
+++ b/Engine/Source/Programs/AutomationTool/AutomationUtils/UE4Build.cs
@@ -794,7 +794,7 @@ namespace AutomationTool
 								int YcIndex = ElementParams.IndexOf(" /Yc\"");
 								if (YcIndex >= 0)
 								{
-									// /Fp&quot;D:\BuildFarm\buildmachine_++depot+UE4\Engine\Intermediate\BuildData\Win64\UE4Editor\Development\SharedPCHs\CoreUObject.h.pch&quot
+									// /Fp&quot;D:\BuildFarm\buildmachine_++depot+UE4\Engine\Intermediate\BuildData\Win64\UE5Editor\Development\SharedPCHs\CoreUObject.h.pch&quot
 									string Fp = " /Fp\"";
 									int FpIndex = ElementParams.IndexOf(Fp, YcIndex);
 									if (FpIndex >= 0)
@@ -803,7 +803,7 @@ namespace AutomationTool
 										if (EndIndex >= 0)
 										{
 											string PCHFileName = ElementParams.Substring(FpIndex + Fp.Length, EndIndex - FpIndex - Fp.Length);
-											if (PCHFileName.Contains(@"\SharedPCHs\") && PCHFileName.Contains(@"\UE4Editor\"))
+											if (PCHFileName.Contains(@"\SharedPCHs\") && PCHFileName.Contains(@"\UE5Editor\"))
 											{
 												Key = "SharedEditorPCH$ " + PCHFileName;
 												CommandUtils.LogLog("Hack: detected Shared PCH, which will use a different key {0}", Key);
diff --git a/Engine/Source/Programs/AutomationTool/BuildGraph/Tasks/CommandletTask.cs b/Engine/Source/Programs/AutomationTool/BuildGraph/Tasks/CommandletTask.cs
--- a/Engine/Source/Programs/AutomationTool/BuildGraph/Tasks/CommandletTask.cs
+++ b/Engine/Source/Programs/AutomationTool/BuildGraph/Tasks/CommandletTask.cs
@@ -36,7 +36,7 @@ namespace AutomationTool.Tasks
 		public string Arguments;
 
 		/// <summary>
-		/// The editor executable to use. Defaults to the development UE4Editor executable for the current platform.
+		/// The editor executable to use. Defaults to the development UE5Editor executable for the current platform.
 		/// </summary>
 		[TaskParameter(Optional = true)]
 		public FileReference EditorExe;
@@ -93,7 +93,7 @@ namespace AutomationTool.Tasks
 			FileReference EditorExe;
 			if(Parameters.EditorExe == null)
 			{
-                EditorExe = new FileReference(HostPlatform.Current.GetUE4ExePath("UE4Editor-Cmd.exe"));
+                EditorExe = new FileReference(HostPlatform.Current.GetUE4ExePath("UE5Editor-Cmd.exe"));
 			}
 			else
 			{
diff --git a/Engine/Source/Programs/AutomationTool/BuildGraph/Tasks/CookTask.cs b/Engine/Source/Programs/AutomationTool/BuildGraph/Tasks/CookTask.cs
--- a/Engine/Source/Programs/AutomationTool/BuildGraph/Tasks/CookTask.cs
+++ b/Engine/Source/Programs/AutomationTool/BuildGraph/Tasks/CookTask.cs
@@ -98,7 +98,7 @@ namespace BuildGraph.Tasks
 			{
 				string[] Maps = (Parameters.Maps == null)? null : Parameters.Maps.Split(new char[]{ '+' });
 				string Arguments = (Parameters.Versioned ? "" : "-Unversioned ") + "-LogCmds=\"LogSavePackage Warning\" " + Parameters.Arguments;
-				CommandUtils.CookCommandlet(ProjectFile, "UE4Editor-Cmd.exe", Maps, null, null, null, Parameters.Platform, Arguments);
+				CommandUtils.CookCommandlet(ProjectFile, "UE5Editor-Cmd.exe", Maps, null, null, null, Parameters.Platform, Arguments);
 			}
 
 			// Find all the cooked files
diff --git a/Engine/Source/Programs/AutomationTool/Gauntlet/Unreal/BuildSource/Gauntlet.UnrealBuildSource.cs b/Engine/Source/Programs/AutomationTool/Gauntlet/Unreal/BuildSource/Gauntlet.UnrealBuildSource.cs
--- a/Engine/Source/Programs/AutomationTool/Gauntlet/Unreal/BuildSource/Gauntlet.UnrealBuildSource.cs
+++ b/Engine/Source/Programs/AutomationTool/Gauntlet/Unreal/BuildSource/Gauntlet.UnrealBuildSource.cs
@@ -541,7 +541,7 @@ namespace Gauntlet
 
 			if (TargetType.UsesEditor())
 			{
-				string ExeFileName = "UE4Editor";
+				string ExeFileName = "UE5Editor";
 				if (TargetConfiguration != UnrealTargetConfiguration.Development)
 				{
 					ExeFileName += string.Format("-{0}-{1}", TargetPlatform.ToString(), TargetConfiguration.ToString());
diff --git a/Engine/Source/Programs/AutomationTool/Scripts/BenchmarkBuild/BenchmarkCookTask.cs b/Engine/Source/Programs/AutomationTool/Scripts/BenchmarkBuild/BenchmarkCookTask.cs
--- a/Engine/Source/Programs/AutomationTool/Scripts/BenchmarkBuild/BenchmarkCookTask.cs
+++ b/Engine/Source/Programs/AutomationTool/Scripts/BenchmarkBuild/BenchmarkCookTask.cs
@@ -62,7 +62,7 @@ namespace AutomationTool.Benchmark
 			if (TaskOptions.HasFlag(DDCTaskOptions.HotDDC))
 			{
 				// will throw an exception if it fails
-				CommandUtils.RunCommandlet(ProjectFile, "UE4Editor-Cmd.exe", "Cook", String.Format("-TargetPlatform={0} ", CookPlatformName));
+				CommandUtils.RunCommandlet(ProjectFile, "UE5Editor-Cmd.exe", "Cook", String.Format("-TargetPlatform={0} ", CookPlatformName));
 			}
 
 			base.PerformPrequisites();
@@ -106,7 +106,7 @@ namespace AutomationTool.Benchmark
 			}
 
 			// will throw an exception if it fails
-			CommandUtils.RunCommandlet(ProjectFile, "UE4Editor-Cmd.exe", "Cook", String.Format("-TargetPlatform={0} {1}", CookPlatformName, ExtraArgs));
+			CommandUtils.RunCommandlet(ProjectFile, "UE5Editor-Cmd.exe", "Cook", String.Format("-TargetPlatform={0} {1}", CookPlatformName, ExtraArgs));
 
 			return true;
 		}
diff --git a/Engine/Source/Programs/AutomationTool/Scripts/BenchmarkBuild/BenchmarkRunEditorTask.cs b/Engine/Source/Programs/AutomationTool/Scripts/BenchmarkBuild/BenchmarkRunEditorTask.cs
--- a/Engine/Source/Programs/AutomationTool/Scripts/BenchmarkBuild/BenchmarkRunEditorTask.cs
+++ b/Engine/Source/Programs/AutomationTool/Scripts/BenchmarkBuild/BenchmarkRunEditorTask.cs
@@ -83,7 +83,7 @@ namespace AutomationTool.Benchmark
 		protected bool RunEditorAndWaitForMapLoad(bool bIsWarming)
 		{
 			string ProjectArg = ProjectFile != null ? ProjectFile.ToString() : "";
-			string EditorPath = HostPlatform.Current.GetUE4ExePath("UE4Editor.exe");
+			string EditorPath = HostPlatform.Current.GetUE4ExePath("UE5Editor.exe");
 			string LogArg = string.Format("-log={0}.log", MakeValidFileName(GetFullTaskName()).Replace(" ", "_"));
 			string Arguments = string.Format("{0} {1} -execcmds=\"automation runtest System.Maps.PIE;Quit\" -stdout -FullStdOutLogOutput -unattended {2}", ProjectArg, EditorArgs, LogArg);
 
diff --git a/Engine/Source/Programs/AutomationTool/Scripts/BuildForUGS.Automation.cs b/Engine/Source/Programs/AutomationTool/Scripts/BuildForUGS.Automation.cs
--- a/Engine/Source/Programs/AutomationTool/Scripts/BuildForUGS.Automation.cs
+++ b/Engine/Source/Programs/AutomationTool/Scripts/BuildForUGS.Automation.cs
@@ -20,7 +20,7 @@ namespace AutomationTool
 			string[] Targets = ParseParamValues("Target");
 			if(Targets.Length == 0)
 			{
-				throw new AutomationException("No targets specified (eg. -Target=\"UE4Editor Win64 Development\")");
+				throw new AutomationException("No targets specified (eg. -Target=\"UE5Editor Win64 Development\")");
 			}
 
 			// Parse the archive path
diff --git a/Engine/Source/Programs/AutomationTool/Scripts/BuildPluginCommand.Automation.cs b/Engine/Source/Programs/AutomationTool/Scripts/BuildPluginCommand.Automation.cs
--- a/Engine/Source/Programs/AutomationTool/Scripts/BuildPluginCommand.Automation.cs
+++ b/Engine/Source/Programs/AutomationTool/Scripts/BuildPluginCommand.Automation.cs
@@ -151,7 +151,7 @@ class BuildPlugin : BuildCommand
 				{
 					CompilePluginWithUBT(HostProjectFile, HostProjectPluginFile, Plugin, "UnrealHeaderTool", TargetType.Program, HostPlatform, UnrealTargetConfiguration.Development, ManifestFileNames, String.Format("{0} -plugin={1}", AdditionalArgs, CommandUtils.MakePathSafeToUseWithCommandLine(HostProjectPluginFile.FullName)));
 				}
-				CompilePluginWithUBT(HostProjectFile, HostProjectPluginFile, Plugin, "UE4Editor", TargetType.Editor, HostPlatform, UnrealTargetConfiguration.Development, ManifestFileNames, AdditionalArgs);
+				CompilePluginWithUBT(HostProjectFile, HostProjectPluginFile, Plugin, "UE5Editor", TargetType.Editor, HostPlatform, UnrealTargetConfiguration.Development, ManifestFileNames, AdditionalArgs);
 			}
 		}
 
diff --git a/Engine/Source/Programs/AutomationTool/Scripts/BuildTarget.Automation.cs b/Engine/Source/Programs/AutomationTool/Scripts/BuildTarget.Automation.cs
--- a/Engine/Source/Programs/AutomationTool/Scripts/BuildTarget.Automation.cs
+++ b/Engine/Source/Programs/AutomationTool/Scripts/BuildTarget.Automation.cs
@@ -14,7 +14,7 @@ namespace AutomationTool
 	[Help("Builds the specified targets and configurations for the specified project.")]
 	[Help("Example BuildTarget -project=QAGame -target=Editor+Game -platform=PS4+XboxOne -configuration=Development.")]
 	[Help("Note: Editor will only ever build for the current platform in a Development config and required tools will be included")]
-	[Help("project=<QAGame>", "Project to build. Will search current path and paths in ueprojectdirs. If omitted will build vanilla UE4Editor")]
+	[Help("project=<QAGame>", "Project to build. Will search current path and paths in ueprojectdirs. If omitted will build vanilla UE5Editor")]
 	[Help("platform=PS4+XboxOne", "Platforms to build, join multiple platforms using +")]
 	[Help("configuration=Development+Test", "Configurations to build, join multiple configurations using +")]
 	[Help("target=Editor+Game", "Targets to build, join multiple targets using +")]
@@ -278,7 +278,7 @@ namespace AutomationTool
 	[Help("Builds the editor for the specified project.")]
 	[Help("Example BuildEditor -project=QAGame")]
 	[Help("Note: Editor will only ever build for the current platform in a Development config and required tools will be included")]
-	[Help("project=<QAGame>", "Project to build. Will search current path and paths in ueprojectdirs. If omitted will build vanilla UE4Editor")]
+	[Help("project=<QAGame>", "Project to build. Will search current path and paths in ueprojectdirs. If omitted will build vanilla UE5Editor")]
 	[Help("notools", "Don't build any tools (UHT, ShaderCompiler, CrashReporter")]
 	class BuildEditor : BuildTarget
 	{
diff --git a/Engine/Source/Programs/AutomationTool/Scripts/CopySharedCookedBuild.Automation.cs b/Engine/Source/Programs/AutomationTool/Scripts/CopySharedCookedBuild.Automation.cs
--- a/Engine/Source/Programs/AutomationTool/Scripts/CopySharedCookedBuild.Automation.cs
+++ b/Engine/Source/Programs/AutomationTool/Scripts/CopySharedCookedBuild.Automation.cs
@@ -111,7 +111,7 @@ class CopySharedCookedBuild : BuildCommand
 				}
 
 				// Get the name of the editor target
-				string EditorTargetName = "UE4Editor";
+				string EditorTargetName = "UE5Editor";
 				if(bIsProjectFile)
 				{
 					string SourceDirectoryName = Path.Combine(Path.GetDirectoryName(ProjectFileName), "Source");
diff --git a/Engine/Source/Programs/AutomationTool/Scripts/FixupRedirects.Automation.cs b/Engine/Source/Programs/AutomationTool/Scripts/FixupRedirects.Automation.cs
--- a/Engine/Source/Programs/AutomationTool/Scripts/FixupRedirects.Automation.cs
+++ b/Engine/Source/Programs/AutomationTool/Scripts/FixupRedirects.Automation.cs
@@ -12,7 +12,7 @@ class FixupRedirects : BuildCommand
 	{
 		var ProjectName = ParseParamValue("project", "");
 
-		var EditorExe = CombinePaths(CmdEnv.LocalRoot, @"Engine/Binaries/Win64/UE4Editor-Cmd.exe");
+		var EditorExe = CombinePaths(CmdEnv.LocalRoot, @"Engine/Binaries/Win64/UE5Editor-Cmd.exe");
 		LogInformation("********** Running FixupRedirects: {0} -run=ResavePackages -unattended -nopause -buildmachine -fixupredirects -autocheckout -autocheckin -projectonly", EditorExe);
 		RunCommandlet(GetCommandletProjectFile(ProjectName), EditorExe, "ResavePackages", "-unattended -nopause -buildmachine -fixupredirects -autocheckout -autocheckin -projectonly");
 	}
diff --git a/Engine/Source/Programs/AutomationTool/Scripts/Localisation.Automation.cs b/Engine/Source/Programs/AutomationTool/Scripts/Localisation.Automation.cs
--- a/Engine/Source/Programs/AutomationTool/Scripts/Localisation.Automation.cs
+++ b/Engine/Source/Programs/AutomationTool/Scripts/Localisation.Automation.cs
@@ -288,7 +288,7 @@ class Localize : BuildCommand
 		// Begin the gather command for each task
 		// These can run in parallel when ParallelGather is enabled
 		{
-			var EditorExe = CombinePaths(CmdEnv.LocalRoot, @"Engine/Binaries/Win64/UE4Editor-Cmd.exe");
+			var EditorExe = CombinePaths(CmdEnv.LocalRoot, @"Engine/Binaries/Win64/UE5Editor-Cmd.exe");
 
 			// Set the common basic editor arguments
 			var EditorArguments = P4Enabled 
@@ -645,7 +645,7 @@ public class ExportMcpTemplates : BuildCommand
 
 	public static void RunExportTemplates(FileReference ProjectFile, bool bCheckoutAndSubmit, bool bOnlyLoc, bool bbNoRobomerge, string CommandletOverride)
 	{
-		string EditorExe = "UE4Editor.exe";
+		string EditorExe = "UE5Editor.exe";
 		EditorExe = HostPlatform.Current.GetUE4ExePath(EditorExe);
 
 		string GameBackendFolder = GetGameBackendFolder(ProjectFile);
diff --git a/Engine/Source/Programs/AutomationTool/Scripts/OpenEditor.Automation.cs b/Engine/Source/Programs/AutomationTool/Scripts/OpenEditor.Automation.cs
--- a/Engine/Source/Programs/AutomationTool/Scripts/OpenEditor.Automation.cs
+++ b/Engine/Source/Programs/AutomationTool/Scripts/OpenEditor.Automation.cs
@@ -12,7 +12,7 @@ namespace AutomationTool
 {
 
 	[Help("Opens the specified project.")]
-	[Help("project=<QAGame>", "Project to open. Will search current path and paths in ueprojectdirs. If omitted will open vanilla UE4Editor")]
+	[Help("project=<QAGame>", "Project to open. Will search current path and paths in ueprojectdirs. If omitted will open vanilla UE5Editor")]
 	class OpenEditor: BuildCommand
 	{
 		// exposed as a property so projects can derive and set this directly
@@ -24,7 +24,7 @@ namespace AutomationTool
 		
 		public override ExitCode Execute()
 		{
-			string EditorPath = HostPlatform.Current.GetUE4ExePath("UE4Editor");
+			string EditorPath = HostPlatform.Current.GetUE4ExePath("UE5Editor");
 
 			string EditorArgs = "";
 
diff --git a/Engine/Source/Programs/AutomationTool/Scripts/ReplaceAssetsUsingManifest.cs b/Engine/Source/Programs/AutomationTool/Scripts/ReplaceAssetsUsingManifest.cs
--- a/Engine/Source/Programs/AutomationTool/Scripts/ReplaceAssetsUsingManifest.cs
+++ b/Engine/Source/Programs/AutomationTool/Scripts/ReplaceAssetsUsingManifest.cs
@@ -17,7 +17,7 @@ public class ReplaceAssetsUsingManifest : BuildCommand
 		// Command parameters - not all required if using existing manifest
 		string ProjectPath = ParseParamValue("ProjectPath");
 		string ManifestFile = ParseParamValue("ManifestFile");
-		string UE4Exe = ParseParamValue("UE4Exe", "UE4Editor-Cmd.exe");
+		string UE4Exe = ParseParamValue("UE4Exe", "UE5Editor-Cmd.exe");
 		string ReplacedPaths = ParseParamValue("ReplacedPaths", "");
 		string ReplacedClasses = ParseParamValue("ReplacedClasses", "");
 		string ExcludedPaths = ParseParamValue("ExcludedPaths", "");
diff --git a/Engine/Source/Programs/AutomationTool/Scripts/RunProjectCommand.Automation.cs b/Engine/Source/Programs/AutomationTool/Scripts/RunProjectCommand.Automation.cs
--- a/Engine/Source/Programs/AutomationTool/Scripts/RunProjectCommand.Automation.cs
+++ b/Engine/Source/Programs/AutomationTool/Scripts/RunProjectCommand.Automation.cs
@@ -798,7 +798,7 @@ public partial class Project : CommandUtils
 		}
 		else
 		{
-			ClientApp = CombinePaths(CmdEnv.LocalRoot, "Engine/Binaries", PlatformName, "UE4Editor.exe");
+			ClientApp = CombinePaths(CmdEnv.LocalRoot, "Engine/Binaries", PlatformName, "UE5Editor.exe");
 			if (!Params.EditorTest)
 			{
 				TempCmdLine += "-game " + Params.MapToRun + " ";
@@ -922,7 +922,7 @@ public partial class Project : CommandUtils
 
 		var SC = DeployContextList[0];
 
-		var ServerApp = CombinePaths(CmdEnv.LocalRoot, "Engine/Binaries/Win64/UE4Editor.exe");
+		var ServerApp = CombinePaths(CmdEnv.LocalRoot, "Engine/Binaries/Win64/UE5Editor.exe");
 		if (ServerParams.Cook)
 		{
 			List<FileReference> Exes = SC.StageTargetPlatform.GetExecutableNames(SC);
@@ -996,7 +996,7 @@ public partial class Project : CommandUtils
 
 	private static IProcessResult RunCookOnTheFlyServer(FileReference ProjectName, string ServerLogFile, string TargetPlatform, string AdditionalCommandLine)
 	{
-		var ServerApp = HostPlatform.Current.GetUE4ExePath("UE4Editor.exe");
+		var ServerApp = HostPlatform.Current.GetUE4ExePath("UE5Editor.exe");
 		var Args = String.Format("{0} -run=cook -cookonthefly -unattended -CrashForUAT -log",
 			CommandUtils.MakePathSafeToUseWithCommandLine(ProjectName.FullName));
 		if (!String.IsNullOrEmpty(ServerLogFile))
diff --git a/Engine/Source/Programs/AutomationTool/Scripts/SyncProject.Automation.cs b/Engine/Source/Programs/AutomationTool/Scripts/SyncProject.Automation.cs
--- a/Engine/Source/Programs/AutomationTool/Scripts/SyncProject.Automation.cs
+++ b/Engine/Source/Programs/AutomationTool/Scripts/SyncProject.Automation.cs
@@ -178,13 +178,13 @@ class SyncProject : BuildCommand
 		else
 		{
 			// See if the engine is in P4 too by checking the p4 location of a local file
-			string LocalEngineFile = CommandUtils.CombinePaths(CmdEnv.LocalRoot, "Engine", "Source", "UE4Editor.target.cs");
+			string LocalEngineFile = CommandUtils.CombinePaths(CmdEnv.LocalRoot, "Engine", "Source", "UE5Editor.target.cs");
 			P4WhereRecord EngineRecord = P4.Where(LocalEngineFile, AllowSpew: false).FirstOrDefault(x => x.DepotFile != null && !x.bUnmap);
 
 			if (!ProjectOnly && P4.FileExistsInDepot(EngineRecord.DepotFile))
 			{
 				// make sure to sync with //workspace/path as it cleans up files if the user has stream switched
-				P4EnginePath = EngineRecord.ClientFile.Replace("Engine/Source/UE4Editor.target.cs", "");
+				P4EnginePath = EngineRecord.ClientFile.Replace("Engine/Source/UE5Editor.target.cs", "");
 				SyncPaths.Add(CommandUtils.CombinePaths(PathSeparator.Slash, P4EnginePath + "*"));
 				SyncPaths.Add(CommandUtils.CombinePaths(PathSeparator.Slash, P4EnginePath, "Engine", "..."));
 			}
diff --git a/Engine/Source/Programs/AutomationTool/Scripts/Tests.Automation.cs b/Engine/Source/Programs/AutomationTool/Scripts/Tests.Automation.cs
--- a/Engine/Source/Programs/AutomationTool/Scripts/Tests.Automation.cs
+++ b/Engine/Source/Programs/AutomationTool/Scripts/Tests.Automation.cs
@@ -683,7 +683,7 @@ class TestTestFarm : BuildCommand
 	{
 		LogInformation("*********************** TestTestFarm");
 
-		string Exe = CombinePaths(CmdEnv.LocalRoot, "Engine", "Binaries", "Win64", "UE4Editor.exe");
+		string Exe = CombinePaths(CmdEnv.LocalRoot, "Engine", "Binaries", "Win64", "UE5Editor.exe");
 		string ClientLogFile = CombinePaths(CmdEnv.LogFolder, "HoverGameRun");
 		string CmdLine = " ../../../Samples/HoverShip/HoverShip.uproject -game -log -abslog=" + ClientLogFile;
 
@@ -793,7 +793,7 @@ class TestLog : BuildCommand
 		PushDir(CombinePaths(CmdEnv.LocalRoot, @"Engine/Binaries/Win64/"));
 		//		string GameServerExe = CombinePaths(CmdEnv.LocalRoot, @"ShooterGame/Binaries/Win64/ShooterGameServer.exe");
 		//		string GameExe = CombinePaths(CmdEnv.LocalRoot, @"ShooterGame/Binaries/Win64/ShooterGame.exe");
-		string EditorExe = CombinePaths(CmdEnv.LocalRoot, @"Engine/Binaries/Win64/UE4Editor.exe");
+		string EditorExe = CombinePaths(CmdEnv.LocalRoot, @"Engine/Binaries/Win64/UE5Editor.exe");
 
 		string ServerLogFile = CombinePaths(CmdEnv.LogFolder, "Server.log");
 
@@ -873,7 +873,7 @@ class TestOSSCommands : BuildCommand
 {
 	public override void ExecuteBuild()
 	{
-		string Exe = CombinePaths(CmdEnv.LocalRoot, "Engine", "Binaries", "Win64", "UE4Editor.exe");
+		string Exe = CombinePaths(CmdEnv.LocalRoot, "Engine", "Binaries", "Win64", "UE5Editor.exe");
 		string CmdLine = "qagame automation-osscommands -game -log -logcmds=" + "\"" + "global none, logonline verbose, loghttp log, LogBlueprintUserMessages log" + "\"";
 		string ClientLogFile = CombinePaths(CmdEnv.LogFolder, String.Format("QALog_{0:yyyy-MM-dd_hh-mm-ss-tt}.txt", DateTime.Now));
 		string LogDirectory = CombinePaths(CmdEnv.LocalRoot, "Saved", "Logs");
@@ -1266,7 +1266,7 @@ class TestKillAll : BuildCommand
 	{
 		LogInformation("*********************** TestKillAll");
 
-		string Exe = CombinePaths(CmdEnv.LocalRoot, "Engine", "Binaries", "Win64", "UE4Editor.exe");
+		string Exe = CombinePaths(CmdEnv.LocalRoot, "Engine", "Binaries", "Win64", "UE5Editor.exe");
 		string ClientLogFile = CombinePaths(CmdEnv.LogFolder, "HoverGameRun");
 		string CmdLine = " ../../../Samples/Sandbox/HoverShip/HoverShip.uproject -game -log -abslog=" + ClientLogFile;
 
@@ -1294,7 +1294,7 @@ class TestStopProcess : BuildCommand
 	{
 		LogInformation("*********************** TestStopProcess");
 
-		string Exe = CombinePaths(CmdEnv.LocalRoot, "Engine", "Binaries", "Win64", "UE4Editor.exe");
+		string Exe = CombinePaths(CmdEnv.LocalRoot, "Engine", "Binaries", "Win64", "UE5Editor.exe");
 		string ClientLogFile = CombinePaths(CmdEnv.LogFolder, "HoverGameRun");
 		string CmdLine = " ../../../Samples/Sandbox/HoverShip/HoverShip.uproject -game -log -ddc=noshared -abslog=" + ClientLogFile;
 
@@ -1328,7 +1328,7 @@ class LookForOverlappingBuildProducts : BuildCommand
 			throw new AutomationException("Source path not found, please use -source=Path");
 		}
 		var Objs = new HashSet<string>( StringComparer.InvariantCultureIgnoreCase );
-//    /Fo&quot;D:\BuildFarm\buildmachine_++depot+UE4\Engine\Intermediate\Build\Win64\UE4Editor\Development\Projects\Module.Projects.cpp.obj&quot;
+//    /Fo&quot;D:\BuildFarm\buildmachine_++depot+UE4\Engine\Intermediate\Build\Win64\UE5Editor\Development\Projects\Module.Projects.cpp.obj&quot;
 		var FileText = ReadAllText(SourcePath);
 		string Start = "/Fo&quot;";
 		string End = "&quot;";
diff --git a/Engine/Source/Programs/CrashReportClient/Private/CrashReportClient.cpp b/Engine/Source/Programs/CrashReportClient/Private/CrashReportClient.cpp
--- a/Engine/Source/Programs/CrashReportClient/Private/CrashReportClient.cpp
+++ b/Engine/Source/Programs/CrashReportClient/Private/CrashReportClient.cpp
@@ -120,7 +120,7 @@ FReply FCrashReportClient::SubmitAndRestart()
 	if (bRunFromLauncher)
 	{
 		// Hacky check to see if this is the editor. Not attempting to relaunch the editor using the Launcher because there is no way to pass the project via OpenLauncher()
-		if (!FPaths::GetCleanFilename(CrashedAppPath).StartsWith(TEXT("UE4Editor")))
+		if (!FPaths::GetCleanFilename(CrashedAppPath).StartsWith(TEXT("UE5Editor")))
 		{
 			// We'll restart Launcher-run processes by having the installed Launcher handle it
 			ILauncherPlatform* LauncherPlatform = FLauncherPlatformModule::Get();
diff --git a/Engine/Source/Programs/Mac/UE4EditorServices/Private/Mac/UE4EditorServicesAppDelegate.cpp b/Engine/Source/Programs/Mac/UE4EditorServices/Private/Mac/UE4EditorServicesAppDelegate.cpp
--- a/Engine/Source/Programs/Mac/UE4EditorServices/Private/Mac/UE4EditorServicesAppDelegate.cpp
+++ b/Engine/Source/Programs/Mac/UE4EditorServices/Private/Mac/UE4EditorServicesAppDelegate.cpp
@@ -60,7 +60,7 @@
 	{
 		if (Iter.Key() == EngineAssociation)
 		{
-			FString BundlePath = Iter.Value() / TEXT("Engine/Binaries/Mac/UE4Editor.app");
+			FString BundlePath = Iter.Value() / TEXT("Engine/Binaries/Mac/UE5Editor.app");
 			return [NSBundle bundleWithPath:BundlePath.GetNSString()];
 		}
 	}
@@ -226,7 +226,7 @@
 	NSString* EnginePath = [self findEngineForUProjectFile:FileURL];
 	if (EnginePath)
 	{
-		NSURL* EditorBundleURL = [NSURL fileURLWithPath:[EnginePath stringByAppendingPathComponent:@"Engine/Binaries/Mac/UE4Editor.app"]];
+		NSURL* EditorBundleURL = [NSURL fileURLWithPath:[EnginePath stringByAppendingPathComponent:@"Engine/Binaries/Mac/UE5Editor.app"]];
 		if (EditorBundleURL)
 		{
 			NSDictionary* Configuration = [NSDictionary dictionaryWithObject:[NSArray arrayWithObjects:[FileURL path], nil] forKey:NSWorkspaceLaunchConfigurationArguments];
@@ -252,7 +252,7 @@
 		NSString* EnginePath = [self findEngineForUProjectFile:FileURL];
 		if (EnginePath)
 		{
-			NSURL* EditorBundleURL = [NSURL fileURLWithPath:[EnginePath stringByAppendingPathComponent:@"Engine/Binaries/Mac/UE4Editor.app"]];
+			NSURL* EditorBundleURL = [NSURL fileURLWithPath:[EnginePath stringByAppendingPathComponent:@"Engine/Binaries/Mac/UE5Editor.app"]];
 			if (EditorBundleURL)
 			{
 				NSDictionary* Configuration = [NSDictionary dictionaryWithObject:[NSArray arrayWithObjects:[FileURL path], @"-game", nil] forKey:NSWorkspaceLaunchConfigurationArguments];
@@ -388,7 +388,7 @@
 	else
 	{
 		EditorDescription = @"Binary Build";
-		NSBundle* EditorBundle = [NSBundle bundleWithPath:[EnginePath stringByAppendingPathComponent:@"Engine/Binaries/Mac/UE4Editor.app"]];
+		NSBundle* EditorBundle = [NSBundle bundleWithPath:[EnginePath stringByAppendingPathComponent:@"Engine/Binaries/Mac/UE5Editor.app"]];
 		if (EditorBundle)
 		{
 			EditorDescription = (NSString*)[[EditorBundle infoDictionary] objectForKey:@"CFBundleVersion"];
@@ -410,7 +410,7 @@
 
 - (BOOL)isAppValidForUProjectFiles:(NSBundle*)AppBundle
 {
-	return AppBundle && ([[AppBundle bundleIdentifier] isEqualToString:@"com.epicgames.UE4Editor"] || [[AppBundle bundleIdentifier] isEqualToString:@"com.epicgames.UE4EditorServices"])
+	return AppBundle && ([[AppBundle bundleIdentifier] isEqualToString:@"com.epicgames.UE5Editor"] || [[AppBundle bundleIdentifier] isEqualToString:@"com.epicgames.UE4EditorServices"])
 					 && [[[AppBundle bundlePath] stringByDeletingLastPathComponent] hasSuffix:@"Engine/Binaries/Mac"]
 					 && [[AppBundle bundlePath] rangeOfString:@".app/Contents/UE4"].location == NSNotFound;
 }
diff --git a/Engine/Source/Programs/UnrealBuildTool/Configuration/TargetRules.cs b/Engine/Source/Programs/UnrealBuildTool/Configuration/TargetRules.cs
--- a/Engine/Source/Programs/UnrealBuildTool/Configuration/TargetRules.cs
+++ b/Engine/Source/Programs/UnrealBuildTool/Configuration/TargetRules.cs
@@ -20,7 +20,7 @@ namespace UnrealBuildTool
 		Game,
 
 		/// <summary>
-		/// Uncooked modular editor executable and DLLs (UE4Editor.exe, UE4Editor*.dll, GameName*.dll)
+		/// Uncooked modular editor executable and DLLs (UE5Editor.exe, UE5Editor*.dll, GameName*.dll)
 		/// </summary>
 		Editor,
 
@@ -304,7 +304,7 @@ namespace UnrealBuildTool
 
 		/// <summary>
 		/// Build all the plugins that we can find, even if they're not enabled. This is particularly useful for content-only projects,
-		/// where you're building the UE4Editor target but running it with a game that enables a plugin.
+		/// where you're building the UE5Editor target but running it with a game that enables a plugin.
 		/// </summary>
 		[Obsolete("bBuildAllPlugins has been deprecated. Use bPrecompile to build all modules which are not part of the target.")]
 		public bool bBuildAllPlugins = false;
diff --git a/Engine/Source/Programs/UnrealBuildTool/Configuration/UEBuildPlatform.cs b/Engine/Source/Programs/UnrealBuildTool/Configuration/UEBuildPlatform.cs
--- a/Engine/Source/Programs/UnrealBuildTool/Configuration/UEBuildPlatform.cs
+++ b/Engine/Source/Programs/UnrealBuildTool/Configuration/UEBuildPlatform.cs
@@ -202,7 +202,7 @@ namespace UnrealBuildTool
 		/// Searches a directory tree for build products to be cleaned.
 		/// </summary>
 		/// <param name="BaseDir">The directory to search</param>
-		/// <param name="NamePrefixes">Target or application names that may appear at the start of the build product name (eg. "UE4Editor", "ShooterGameEditor")</param>
+		/// <param name="NamePrefixes">Target or application names that may appear at the start of the build product name (eg. "UE5Editor", "ShooterGameEditor")</param>
 		/// <param name="NameSuffixes">Suffixes which may appear at the end of the build product name</param>
 		/// <param name="FilesToClean">List to receive a list of files to be cleaned</param>
 		/// <param name="DirectoriesToClean">List to receive a list of directories to be cleaned</param>
@@ -244,7 +244,7 @@ namespace UnrealBuildTool
 		/// Determines if a filename is a default UBT build product
 		/// </summary>
 		/// <param name="FileName">The name to check</param>
-		/// <param name="NamePrefixes">Target or application names that may appear at the start of the build product name (eg. "UE4Editor", "ShooterGameEditor")</param>
+		/// <param name="NamePrefixes">Target or application names that may appear at the start of the build product name (eg. "UE5Editor", "ShooterGameEditor")</param>
 		/// <param name="NameSuffixes">Suffixes which may appear at the end of the build product name</param>
 		/// <returns>True if the substring matches the name of a build product, false otherwise</returns>
 		public static bool IsDefaultBuildProduct(string FileName, string[] NamePrefixes, string[] NameSuffixes)
@@ -258,16 +258,16 @@ namespace UnrealBuildTool
 		/// Determines if the given name is a build product for a target.
 		/// </summary>
 		/// <param name="FileName">The name to check</param>
-		/// <param name="NamePrefixes">Target or application names that may appear at the start of the build product name (eg. "UE4Editor", "ShooterGameEditor")</param>
+		/// <param name="NamePrefixes">Target or application names that may appear at the start of the build product name (eg. "UE5Editor", "ShooterGameEditor")</param>
 		/// <param name="NameSuffixes">Suffixes which may appear at the end of the build product name</param>
 		/// <returns>True if the string matches the name of a build product, false otherwise</returns>
 		public abstract bool IsBuildProduct(string FileName, string[] NamePrefixes, string[] NameSuffixes);
 
 		/// <summary>
-		/// Determines if a string is in the canonical name of a UE build product, with a specific extension (eg. "UE4Editor-Win64-Debug.exe" or "UE4Editor-ModuleName-Win64-Debug.dll"). 
+		/// Determines if a string is in the canonical name of a UE build product, with a specific extension (eg. "UE5Editor-Win64-Debug.exe" or "UE5Editor-ModuleName-Win64-Debug.dll"). 
 		/// </summary>
 		/// <param name="FileName">The file name to check</param>
-		/// <param name="NamePrefixes">Target or application names that may appear at the start of the build product name (eg. "UE4Editor", "ShooterGameEditor")</param>
+		/// <param name="NamePrefixes">Target or application names that may appear at the start of the build product name (eg. "UE5Editor", "ShooterGameEditor")</param>
 		/// <param name="NameSuffixes">Suffixes which may appear at the end of the build product name</param>
 		/// <param name="Extension">The extension to check for</param>
 		/// <returns>True if the string matches the name of a build product, false otherwise</returns>
@@ -277,12 +277,12 @@ namespace UnrealBuildTool
 		}
 
 		/// <summary>
-		/// Determines if a substring is in the canonical name of a UE build product, with a specific extension (eg. "UE4Editor-Win64-Debug.exe" or "UE4Editor-ModuleName-Win64-Debug.dll"). 
+		/// Determines if a substring is in the canonical name of a UE build product, with a specific extension (eg. "UE5Editor-Win64-Debug.exe" or "UE5Editor-ModuleName-Win64-Debug.dll"). 
 		/// </summary>
 		/// <param name="FileName">The name to check</param>
 		/// <param name="Index">Index of the first character to be checked</param>
 		/// <param name="Count">Number of characters of the substring to check</param>
-		/// <param name="NamePrefixes">Target or application names that may appear at the start of the build product name (eg. "UE4Editor", "ShooterGameEditor")</param>
+		/// <param name="NamePrefixes">Target or application names that may appear at the start of the build product name (eg. "UE5Editor", "ShooterGameEditor")</param>
 		/// <param name="NameSuffixes">Suffixes which may appear at the end of the build product name</param>
 		/// <param name="Extension">The extension to check for</param>
 		/// <returns>True if the substring matches the name of a build product, false otherwise</returns>
@@ -297,12 +297,12 @@ namespace UnrealBuildTool
 		}
 
 		/// <summary>
-		/// Determines if a substring is in the canonical name of a UE build product, excluding extension or other decoration (eg. "UE4Editor-Win64-Debug" or "UE4Editor-ModuleName-Win64-Debug"). 
+		/// Determines if a substring is in the canonical name of a UE build product, excluding extension or other decoration (eg. "UE5Editor-Win64-Debug" or "UE5Editor-ModuleName-Win64-Debug"). 
 		/// </summary>
 		/// <param name="FileName">The name to check</param>
 		/// <param name="Index">Index of the first character to be checked</param>
 		/// <param name="Count">Number of characters of the substring to check</param>
-		/// <param name="NamePrefixes">Target or application names that may appear at the start of the build product name (eg. "UE4Editor", "ShooterGameEditor")</param>
+		/// <param name="NamePrefixes">Target or application names that may appear at the start of the build product name (eg. "UE5Editor", "ShooterGameEditor")</param>
 		/// <param name="NameSuffixes">Suffixes which may appear at the end of the build product name</param>
 		/// <returns>True if the substring matches the name of a build product, false otherwise</returns>
 		public static bool IsBuildProductName(string FileName, int Index, int Count, string[] NamePrefixes, string[] NameSuffixes)
diff --git a/Engine/Source/Programs/UnrealBuildTool/Configuration/UEBuildTarget.cs b/Engine/Source/Programs/UnrealBuildTool/Configuration/UEBuildTarget.cs
--- a/Engine/Source/Programs/UnrealBuildTool/Configuration/UEBuildTarget.cs
+++ b/Engine/Source/Programs/UnrealBuildTool/Configuration/UEBuildTarget.cs
@@ -808,7 +808,7 @@ namespace UnrealBuildTool
 					BaseTargetName = "UE4Game";
 					break;
 				case TargetType.Editor:
-					BaseTargetName = "UE4Editor";
+					BaseTargetName = "UE5Editor";
 					break;
 				case TargetType.Client:
 					BaseTargetName = "UE4Client";
@@ -921,7 +921,7 @@ namespace UnrealBuildTool
 		public TargetType TargetType;
 
 		/// <summary>
-		/// The name of the application the target is part of. For targets with bUseSharedBuildEnvironment = true, this is typically the name of the base application, eg. UE4Editor for any game editor.
+		/// The name of the application the target is part of. For targets with bUseSharedBuildEnvironment = true, this is typically the name of the base application, eg. UE5Editor for any game editor.
 		/// </summary>
 		public string AppName;
 
@@ -1149,7 +1149,7 @@ namespace UnrealBuildTool
 				case TargetType.Server:
 					return "UE4Server";
 				case TargetType.Editor:
-					return "UE4Editor";
+					return "UE5Editor";
 				default:
 					throw new BuildException("Invalid target type ({0})", (int)Type);
 			}
diff --git a/Engine/Source/Programs/UnrealBuildTool/Modes/CleanMode.cs b/Engine/Source/Programs/UnrealBuildTool/Modes/CleanMode.cs
--- a/Engine/Source/Programs/UnrealBuildTool/Modes/CleanMode.cs
+++ b/Engine/Source/Programs/UnrealBuildTool/Modes/CleanMode.cs
@@ -13,7 +13,7 @@ namespace UnrealBuildTool
 {
 	/// <summary>
 	/// Cleans build products and intermediates for the target. This deletes files which are named consistently with the target being built
-	/// (e.g. UE4Editor-Foo-Win64-Debug.dll) rather than an actual record of previous build products.
+	/// (e.g. UE5Editor-Foo-Win64-Debug.dll) rather than an actual record of previous build products.
 	/// </summary>
 	[ToolMode("Clean", ToolModeOptions.XmlConfig | ToolModeOptions.BuildPlatforms | ToolModeOptions.SingleInstance)]
 	class CleanMode : ToolMode
diff --git a/Engine/Source/Programs/UnrealBuildTool/Platform/HoloLens/UEBuildHoloLens.cs b/Engine/Source/Programs/UnrealBuildTool/Platform/HoloLens/UEBuildHoloLens.cs
--- a/Engine/Source/Programs/UnrealBuildTool/Platform/HoloLens/UEBuildHoloLens.cs
+++ b/Engine/Source/Programs/UnrealBuildTool/Platform/HoloLens/UEBuildHoloLens.cs
@@ -261,7 +261,7 @@ namespace UnrealBuildTool
 		/// Determines if the given name is a build product for a target.
 		/// </summary>
 		/// <param name="FileName">The name to check</param>
-		/// <param name="NamePrefixes">Target or application names that may appear at the start of the build product name (eg. "UE4Editor", "ShooterGameEditor")</param>
+		/// <param name="NamePrefixes">Target or application names that may appear at the start of the build product name (eg. "UE5Editor", "ShooterGameEditor")</param>
 		/// <param name="NameSuffixes">Suffixes which may appear at the end of the build product name</param>
 		/// <returns>True if the string matches the name of a build product, false otherwise</returns>
 		public override bool IsBuildProduct(string FileName, string[] NamePrefixes, string[] NameSuffixes)
diff --git a/Engine/Source/Programs/UnrealBuildTool/Platform/IOS/UEBuildIOS.cs b/Engine/Source/Programs/UnrealBuildTool/Platform/IOS/UEBuildIOS.cs
--- a/Engine/Source/Programs/UnrealBuildTool/Platform/IOS/UEBuildIOS.cs
+++ b/Engine/Source/Programs/UnrealBuildTool/Platform/IOS/UEBuildIOS.cs
@@ -811,7 +811,7 @@ namespace UnrealBuildTool
 		/// Determines if the given name is a build product for a target.
 		/// </summary>
 		/// <param name="FileName">The name to check</param>
-		/// <param name="NamePrefixes">Target or application names that may appear at the start of the build product name (eg. "UE4Editor", "ShooterGameEditor")</param>
+		/// <param name="NamePrefixes">Target or application names that may appear at the start of the build product name (eg. "UE5Editor", "ShooterGameEditor")</param>
 		/// <param name="NameSuffixes">Suffixes which may appear at the end of the build product name</param>
 		/// <returns>True if the string matches the name of a build product, false otherwise</returns>
 		public override bool IsBuildProduct(string FileName, string[] NamePrefixes, string[] NameSuffixes)
diff --git a/Engine/Source/Programs/UnrealBuildTool/Platform/Linux/UEBuildLinux.cs b/Engine/Source/Programs/UnrealBuildTool/Platform/Linux/UEBuildLinux.cs
--- a/Engine/Source/Programs/UnrealBuildTool/Platform/Linux/UEBuildLinux.cs
+++ b/Engine/Source/Programs/UnrealBuildTool/Platform/Linux/UEBuildLinux.cs
@@ -291,7 +291,7 @@ namespace UnrealBuildTool
 		/// Determines if the given name is a build product for a target.
 		/// </summary>
 		/// <param name="FileName">The name to check</param>
-		/// <param name="NamePrefixes">Target or application names that may appear at the start of the build product name (eg. "UE4Editor", "ShooterGameEditor")</param>
+		/// <param name="NamePrefixes">Target or application names that may appear at the start of the build product name (eg. "UE5Editor", "ShooterGameEditor")</param>
 		/// <param name="NameSuffixes">Suffixes which may appear at the end of the build product name</param>
 		/// <returns>True if the string matches the name of a build product, false otherwise</returns>
 		public override bool IsBuildProduct(string FileName, string[] NamePrefixes, string[] NameSuffixes)
diff --git a/Engine/Source/Programs/UnrealBuildTool/Platform/Mac/MacToolChain.cs b/Engine/Source/Programs/UnrealBuildTool/Platform/Mac/MacToolChain.cs
--- a/Engine/Source/Programs/UnrealBuildTool/Platform/Mac/MacToolChain.cs
+++ b/Engine/Source/Programs/UnrealBuildTool/Platform/Mac/MacToolChain.cs
@@ -1055,7 +1055,7 @@ namespace UnrealBuildTool
 					}
 					AppendMacLine(FinalizeAppBundleScript, FormatCopyCommand(CustomIcon, String.Format("{0}.app/Contents/Resources/{1}.icns", ExeName, GameName)));
 
-					if (ExeName.StartsWith("UE4Editor"))
+					if (ExeName.StartsWith("UE5Editor"))
 					{
 						AppendMacLine(FinalizeAppBundleScript, FormatCopyCommand(String.Format("{0}/Runtime/Launch/Resources/Mac/UProject.icns", EngineSourcePath), String.Format("{0}.app/Contents/Resources/UProject.icns", ExeName)));
 					}
@@ -1401,7 +1401,7 @@ namespace UnrealBuildTool
 				}
 			}
 
-			if ((BinaryLinkEnvironment.bIsBuildingDLL && (Options & MacToolChainOptions.OutputDylib) == 0) || (BinaryLinkEnvironment.bIsBuildingConsoleApplication && Executable.Name.Contains("UE4Editor") && Executable.Name.EndsWith("-Cmd")))
+			if ((BinaryLinkEnvironment.bIsBuildingDLL && (Options & MacToolChainOptions.OutputDylib) == 0) || (BinaryLinkEnvironment.bIsBuildingConsoleApplication && Executable.Name.Contains("UE5Editor") && Executable.Name.EndsWith("-Cmd")))
 			{
 				return OutputFiles;
 			}
diff --git a/Engine/Source/Programs/UnrealBuildTool/Platform/Mac/UEBuildMac.cs b/Engine/Source/Programs/UnrealBuildTool/Platform/Mac/UEBuildMac.cs
--- a/Engine/Source/Programs/UnrealBuildTool/Platform/Mac/UEBuildMac.cs
+++ b/Engine/Source/Programs/UnrealBuildTool/Platform/Mac/UEBuildMac.cs
@@ -182,7 +182,7 @@ namespace UnrealBuildTool
 		/// Determines if the given name is a build product for a target.
 		/// </summary>
 		/// <param name="FileName">The name to check</param>
-		/// <param name="NamePrefixes">Target or application names that may appear at the start of the build product name (eg. "UE4Editor", "ShooterGameEditor")</param>
+		/// <param name="NamePrefixes">Target or application names that may appear at the start of the build product name (eg. "UE5Editor", "ShooterGameEditor")</param>
 		/// <param name="NameSuffixes">Suffixes which may appear at the end of the build product name</param>
 		/// <returns>True if the string matches the name of a build product, false otherwise</returns>
 		public override bool IsBuildProduct(string FileName, string[] NamePrefixes, string[] NameSuffixes)
diff --git a/Engine/Source/Programs/UnrealBuildTool/Platform/Windows/UEBuildWindows.cs b/Engine/Source/Programs/UnrealBuildTool/Platform/Windows/UEBuildWindows.cs
--- a/Engine/Source/Programs/UnrealBuildTool/Platform/Windows/UEBuildWindows.cs
+++ b/Engine/Source/Programs/UnrealBuildTool/Platform/Windows/UEBuildWindows.cs
@@ -1924,7 +1924,7 @@ namespace UnrealBuildTool
 		/// Determines if the given name is a build product for a target.
 		/// </summary>
 		/// <param name="FileName">The name to check</param>
-		/// <param name="NamePrefixes">Target or application names that may appear at the start of the build product name (eg. "UE4Editor", "ShooterGameEditor")</param>
+		/// <param name="NamePrefixes">Target or application names that may appear at the start of the build product name (eg. "UE5Editor", "ShooterGameEditor")</param>
 		/// <param name="NameSuffixes">Suffixes which may appear at the end of the build product name</param>
 		/// <returns>True if the string matches the name of a build product, false otherwise</returns>
 		public override bool IsBuildProduct(string FileName, string[] NamePrefixes, string[] NameSuffixes)
diff --git a/Engine/Source/Programs/UnrealBuildTool/ProjectFiles/CMake/CMakefileGenerator.cs b/Engine/Source/Programs/UnrealBuildTool/ProjectFiles/CMake/CMakefileGenerator.cs
--- a/Engine/Source/Programs/UnrealBuildTool/ProjectFiles/CMake/CMakefileGenerator.cs
+++ b/Engine/Source/Programs/UnrealBuildTool/ProjectFiles/CMake/CMakefileGenerator.cs
@@ -694,7 +694,7 @@ namespace UnrealBuildTool
 			{
 				// Editor or game builds get all target configs
 				// The game project editor or game get all configs
-				if ((TargetName.StartsWith("UE4Editor") && !TargetName.StartsWith("UE4EditorServices")) ||
+				if ((TargetName.StartsWith("UE5Editor") && !TargetName.StartsWith("UE4EditorServices")) ||
 					TargetName.StartsWith("UE4Game") ||
 					(IsProjectBuild && TargetName.StartsWith(GameProjectName)))
 				{
diff --git a/Engine/Source/Programs/UnrealBuildTool/ProjectFiles/CodeLite/CodeLiteGenerator.cs b/Engine/Source/Programs/UnrealBuildTool/ProjectFiles/CodeLite/CodeLiteGenerator.cs
--- a/Engine/Source/Programs/UnrealBuildTool/ProjectFiles/CodeLite/CodeLiteGenerator.cs
+++ b/Engine/Source/Programs/UnrealBuildTool/ProjectFiles/CodeLite/CodeLiteGenerator.cs
@@ -190,7 +190,7 @@ namespace UnrealBuildTool
 							if (ProjectName.Equals("UE4Client") ||
 								ProjectName.Equals("UE4Server") ||
 								ProjectName.Equals("UE4Game") ||
-								ProjectName.Equals("UE4Editor"))
+								ProjectName.Equals("UE5Editor"))
 							{
 								if (CodeLiteWorkspaceTargetEngine == null)
 								{
diff --git a/Engine/Source/Programs/UnrealBuildTool/ProjectFiles/CodeLite/CodeLiteProject.cs b/Engine/Source/Programs/UnrealBuildTool/ProjectFiles/CodeLite/CodeLiteProject.cs
--- a/Engine/Source/Programs/UnrealBuildTool/ProjectFiles/CodeLite/CodeLiteProject.cs
+++ b/Engine/Source/Programs/UnrealBuildTool/ProjectFiles/CodeLite/CodeLiteProject.cs
@@ -55,7 +55,7 @@ namespace UnrealBuildTool
 				GameWorkingDirectory = Path.Combine (Path.GetDirectoryName (OnlyGameProject.FullName), "Binaries", ProjectPlatformName);
 			}
 			//
-			// Build the working directory of the UE4Editor executable.
+			// Build the working directory of the UE5Editor executable.
 			//
 			string UE4EditorWorkingDirectory = Path.Combine(EngineRootDirectory, "Binaries", ProjectPlatformName);
 
@@ -125,7 +125,7 @@ namespace UnrealBuildTool
 						if (TargetName.Equals("UE4Client") ||
 							TargetName.Equals("UE4Server") ||
 							TargetName.Equals("UE4Game") ||
-							TargetName.Equals("UE4Editor"))
+							TargetName.Equals("UE5Editor"))
 						{
 							int Idx = UnrealBuildTool.EngineDirectory.FullName.Length;
 							CurrentFilePath = Path.GetDirectoryName(Path.GetFullPath(CurrentFile.Reference.FullName)).Substring(Idx);
@@ -259,7 +259,7 @@ namespace UnrealBuildTool
 						}
 						else if (ProjectTargetType == TargetType.Editor)
 						{
-							ExecutableToRun = "./UE4Editor";
+							ExecutableToRun = "./UE5Editor";
 							if ((CurConf == UnrealTargetConfiguration.Debug) || 
 								(CurConf == UnrealTargetConfiguration.Shipping) || 
 								(CurConf == UnrealTargetConfiguration.Test))
@@ -287,14 +287,14 @@ namespace UnrealBuildTool
 						}
 						else if (ProjectTargetType == TargetType.Editor)
 						{
-							ExecutableToRun = "./UE4Editor";
+							ExecutableToRun = "./UE5Editor";
 							if ((CurConf == UnrealTargetConfiguration.Debug) || 
 								(CurConf == UnrealTargetConfiguration.Shipping) || 
 								(CurConf == UnrealTargetConfiguration.Test))
 							{
 								ExecutableToRun += PlatformConfiguration;
 							}
-							ExecutableToRun += ".app/Contents/MacOS/UE4Editor";
+							ExecutableToRun += ".app/Contents/MacOS/UE5Editor";
 							if ((CurConf != UnrealTargetConfiguration.Development) && (CurConf != UnrealTargetConfiguration.DebugGame))
 							{
 								ExecutableToRun += PlatformConfiguration;
@@ -315,7 +315,7 @@ namespace UnrealBuildTool
 						}
 						else if (ProjectTargetType == TargetType.Editor)
 						{
-							ExecutableToRun = "UE4Editor";
+							ExecutableToRun = "UE5Editor";
 							if ((CurConf == UnrealTargetConfiguration.Debug) || 
 								(CurConf == UnrealTargetConfiguration.Shipping) || 
 								(CurConf == UnrealTargetConfiguration.Test))
@@ -353,7 +353,7 @@ namespace UnrealBuildTool
 					} 
 					else if (ProjectTargetType == TargetType.Editor) 
 					{
-						if (TargetName != "UE4Editor" && GameProjectFile != "")
+						if (TargetName != "UE5Editor" && GameProjectFile != "")
 						{
 							string commandArguments = "\"" + GameProjectFile + "\"" + " -game";
 							XAttribute CommandArguments = new XAttribute("CommandArguments", commandArguments);
diff --git a/Engine/Source/Programs/UnrealBuildTool/ProjectFiles/KDevelop/KDevelopGenerator.cs b/Engine/Source/Programs/UnrealBuildTool/ProjectFiles/KDevelop/KDevelopGenerator.cs
--- a/Engine/Source/Programs/UnrealBuildTool/ProjectFiles/KDevelop/KDevelopGenerator.cs
+++ b/Engine/Source/Programs/UnrealBuildTool/ProjectFiles/KDevelop/KDevelopGenerator.cs
@@ -164,14 +164,14 @@ namespace UnrealBuildTool
 
 			FileContent.Append("Title=BuildMeFirst\n\n");
 			FileContent.Append("[CustomBuildSystem][BuildConfig0][ToolBuild]\n");
-			FileContent.Append("Arguments=-f Makefile UE4Editor UE4Game ShaderCompileWorker UnrealLightmass UnrealPak\n");
+			FileContent.Append("Arguments=-f Makefile UE5Editor UE4Game ShaderCompileWorker UnrealLightmass UnrealPak\n");
 			FileContent.Append("Enabled=true\n");
 			FileContent.Append("Environment=\n");
 			FileContent.Append("Executable=make\n");
 			FileContent.Append("Type=0\n\n");
 
 			FileContent.Append("[CustomBuildSystem][BuildConfig0][ToolClean]\n");
-			FileContent.Append("Arguments=-f Makefile UE4Editor UE4Game ShaderCompileWorker UnrealLightmass UnrealPak -clean\n");
+			FileContent.Append("Arguments=-f Makefile UE5Editor UE4Game ShaderCompileWorker UnrealLightmass UnrealPak -clean\n");
 			FileContent.Append("Enabled=true\n");
 			FileContent.Append("Environment=\n");
 			FileContent.Append("Executable=make\n");
diff --git a/Engine/Source/Programs/UnrealBuildTool/ProjectFiles/Make/MakefileGenerator.cs b/Engine/Source/Programs/UnrealBuildTool/ProjectFiles/Make/MakefileGenerator.cs
--- a/Engine/Source/Programs/UnrealBuildTool/ProjectFiles/Make/MakefileGenerator.cs
+++ b/Engine/Source/Programs/UnrealBuildTool/ProjectFiles/Make/MakefileGenerator.cs
@@ -127,8 +127,8 @@ namespace UnrealBuildTool
 			MakefileContent.Append("\n\n" + BuildCommand + ProjectBuildCommand + "\n" +
 				"all: StandardSet\n\n" +
 				"RequiredTools: CrashReportClient-Linux-Shipping CrashReportClientEditor-Linux-Shipping ShaderCompileWorker UnrealLightmass\n\n" +
-				"StandardSet: RequiredTools UnrealFrontend UE4Editor UnrealInsights\n\n" +
-				"DebugSet: RequiredTools UnrealFrontend-Linux-Debug UE4Editor-Linux-Debug\n\n"
+				"StandardSet: RequiredTools UnrealFrontend UE5Editor UnrealInsights\n\n" +
+				"DebugSet: RequiredTools UnrealFrontend-Linux-Debug UE5Editor-Linux-Debug\n\n"
 			);
 
 			foreach (ProjectFile Project in GeneratedProjectFiles)
diff --git a/Engine/Source/Programs/UnrealBuildTool/ProjectFiles/ProjectFileGenerator.cs b/Engine/Source/Programs/UnrealBuildTool/ProjectFiles/ProjectFileGenerator.cs
--- a/Engine/Source/Programs/UnrealBuildTool/ProjectFiles/ProjectFileGenerator.cs
+++ b/Engine/Source/Programs/UnrealBuildTool/ProjectFiles/ProjectFileGenerator.cs
@@ -324,7 +324,7 @@ namespace UnrealBuildTool
 		/// <summary>
 		/// Name of the UE4 engine project that contains all of the engine code, config files and other files
 		/// </summary>
-		public const string EngineProjectFileNameBase = "UE4";
+		public const string EngineProjectFileNameBase = "UE5";
 
 		/// <summary>
 		/// Name of the UE4 enterprise project that contains all of the enterprise code, config files and other files
@@ -363,7 +363,7 @@ namespace UnrealBuildTool
 		/// Name of the master project file -- for example, the base file name for the Visual Studio solution file, or the Xcode project file on Mac.
 		/// </summary>
 		[XmlConfigFile]
-		protected string MasterProjectName = "UE4";
+		protected string MasterProjectName = "UE5";
 
 		/// <summary>
 		/// If true, sets the master project name according to the name of the folder it is in.
diff --git a/Engine/Source/Programs/UnrealBuildTool/ProjectFiles/VisualStudio/VCProject.cs b/Engine/Source/Programs/UnrealBuildTool/ProjectFiles/VisualStudio/VCProject.cs
--- a/Engine/Source/Programs/UnrealBuildTool/ProjectFiles/VisualStudio/VCProject.cs
+++ b/Engine/Source/Programs/UnrealBuildTool/ProjectFiles/VisualStudio/VCProject.cs
@@ -1566,7 +1566,14 @@ namespace UnrealBuildTool
 					string BaseExeName = TargetName;
 					if (!bShouldCompileMonolithic && TargetRulesObject.Type != TargetType.Program && TargetRulesObject.BuildEnvironment != TargetBuildEnvironment.Unique)
 					{
-						BaseExeName = "UE4" + TargetRulesObject.Type.ToString();
+						if (TargetRulesObject.Type == TargetType.Editor)
+						{
+							BaseExeName = "UE5Editor";
+						}
+						else
+						{
+							BaseExeName = "UE4" + TargetRulesObject.Type.ToString();
+						}
 					}
 
 					// Make the output file path
@@ -1709,7 +1716,7 @@ namespace UnrealBuildTool
 								DebugOptions += UProjectPath;
 								DebugOptions += " -skipcompile";
 							}
-							else if (TargetRulesObject.Type == TargetType.Editor && ProjectName != "UE4")
+							else if (TargetRulesObject.Type == TargetType.Editor && ProjectName != ProjectFileGenerator.EngineProjectFileNameBase)
 							{
 								DebugOptions += ProjectName;
 							}
diff --git a/Engine/Source/Programs/UnrealBuildTool/ProjectFiles/VisualStudio/VCProjectFileGenerator.cs b/Engine/Source/Programs/UnrealBuildTool/ProjectFiles/VisualStudio/VCProjectFileGenerator.cs
--- a/Engine/Source/Programs/UnrealBuildTool/ProjectFiles/VisualStudio/VCProjectFileGenerator.cs
+++ b/Engine/Source/Programs/UnrealBuildTool/ProjectFiles/VisualStudio/VCProjectFileGenerator.cs
@@ -416,7 +416,7 @@ namespace UnrealBuildTool
 			IDictionary<MasterProjectFolder, Guid> Guids = new Dictionary<MasterProjectFolder, Guid>();
 			foreach (MasterProjectFolder Folder in RootFolder.SubFolders)
 			{
-				GenerateProjectFolderGuids("UE4", Folder, Guids);
+				GenerateProjectFolderGuids("UE5", Folder, Guids);
 			}
 			return Guids;
 		}
diff --git a/Engine/Source/Programs/UnrealBuildTool/ProjectFiles/VisualStudioCode/VSCodeProjectFileGenerator.cs b/Engine/Source/Programs/UnrealBuildTool/ProjectFiles/VisualStudioCode/VSCodeProjectFileGenerator.cs
--- a/Engine/Source/Programs/UnrealBuildTool/ProjectFiles/VisualStudioCode/VSCodeProjectFileGenerator.cs
+++ b/Engine/Source/Programs/UnrealBuildTool/ProjectFiles/VisualStudioCode/VSCodeProjectFileGenerator.cs
@@ -965,7 +965,7 @@ namespace UnrealBuildTool
 							{
 								if (Target.Type == TargetRules.TargetType.Editor)
 								{
-									if (InProject.Name != "UE4")
+									if (InProject.Name != "UE5")
 									{
 										if (bForeignProject)
 										{
@@ -1238,7 +1238,7 @@ namespace UnrealBuildTool
 					// This is also known as ${workspaceRoot}
 					WorkspaceFile.BeginObject();
 					{
-						string ProjectName = bForeignProject ? GameProjectName : "UE4";
+						string ProjectName = bForeignProject ? GameProjectName : "UE5";
 						WorkspaceFile.AddField("name", ProjectName);
 						WorkspaceFile.AddField("path", ".");
 					}
@@ -1249,7 +1249,7 @@ namespace UnrealBuildTool
 					{
 						WorkspaceFile.BeginObject();
 						{
-							WorkspaceFile.AddField("name", "UE4");
+							WorkspaceFile.AddField("name", "UE5");
 							WorkspaceFile.AddField("path", MakeUnquotedPathString(UnrealBuildTool.RootDirectory, EPathType.Absolute));
 						}
 						WorkspaceFile.EndObject();
@@ -1287,7 +1287,7 @@ namespace UnrealBuildTool
 
 			WorkspaceFile.EndRootObject();
 
-			string WorkspaceName = bForeignProject ? GameProjectName : "UE4";
+			string WorkspaceName = bForeignProject ? GameProjectName : "UE5";
 			WorkspaceFile.Write(FileReference.Combine(MasterProjectPath, WorkspaceName + ".code-workspace"));
 		}
 
diff --git a/Engine/Source/Programs/UnrealBuildTool/ProjectFiles/Xcode/XcodeProject.cs b/Engine/Source/Programs/UnrealBuildTool/ProjectFiles/Xcode/XcodeProject.cs
--- a/Engine/Source/Programs/UnrealBuildTool/ProjectFiles/Xcode/XcodeProject.cs
+++ b/Engine/Source/Programs/UnrealBuildTool/ProjectFiles/Xcode/XcodeProject.cs
@@ -1494,7 +1494,11 @@ namespace UnrealBuildTool
 											if (!bShouldCompileMonolithic && ProjectTarget.TargetRules.Type != TargetType.Program)
 											{
 												// Figure out what the compiled binary will be called so that we can point the IDE to the correct file
-												if (ProjectTarget.TargetRules.Type != TargetType.Game)
+												if (ProjectTarget.TargetRules.Type == TargetType.Editor)
+												{
+													ExeName = "UE5Editor";
+												}
+												else if (ProjectTarget.TargetRules.Type != TargetType.Game)
 												{
 													ExeName = "UE4" + ProjectTarget.TargetRules.Type.ToString();
 												}
diff --git a/Engine/Source/Programs/UnrealBuildTool/System/ExternalExecution.cs b/Engine/Source/Programs/UnrealBuildTool/System/ExternalExecution.cs
--- a/Engine/Source/Programs/UnrealBuildTool/System/ExternalExecution.cs
+++ b/Engine/Source/Programs/UnrealBuildTool/System/ExternalExecution.cs
@@ -1166,7 +1166,7 @@ namespace UnrealBuildTool
 
 					Progress.Write(1, 3);
 
-					string ActualTargetName = String.IsNullOrEmpty(TargetName) ? "UE4" : TargetName;
+					string ActualTargetName = String.IsNullOrEmpty(TargetName) ? "UE5" : TargetName;
 					Log.TraceInformation("Parsing headers for {0}", ActualTargetName);
 
 					FileReference HeaderToolPath = GetHeaderToolPath(HeaderToolReceipt);
diff --git a/Engine/Source/Programs/UnrealBuildTool/System/HotReload.cs b/Engine/Source/Programs/UnrealBuildTool/System/HotReload.cs
--- a/Engine/Source/Programs/UnrealBuildTool/System/HotReload.cs
+++ b/Engine/Source/Programs/UnrealBuildTool/System/HotReload.cs
@@ -276,7 +276,7 @@ namespace UnrealBuildTool
 			// this code must be able to execute before we create or load module rules DLLs so that hot reload can work with bUseUBTMakefiles
 			if (TargetDesc.Name.EndsWith("Editor", StringComparison.OrdinalIgnoreCase))
 			{
-				string EditorBaseFileName = "UE4Editor";
+				string EditorBaseFileName = "UE5Editor";
 				if (TargetDesc.Configuration != UnrealTargetConfiguration.Development)
 				{
 					EditorBaseFileName = String.Format("{0}-{1}-{2}", EditorBaseFileName, TargetDesc.Platform, TargetDesc.Configuration);
@@ -667,7 +667,7 @@ namespace UnrealBuildTool
 			int HyphenIdx = FileName.IndexOf('-');
 			if (HyphenIdx == -1)
 			{
-				throw new BuildException("Hot-reloadable files are expected to contain a hyphen, eg. UE4Editor-Core");
+				throw new BuildException("Hot-reloadable files are expected to contain a hyphen, eg. UE5Editor-Core");
 			}
 
 			int NameEndIdx = HyphenIdx + 1;
diff --git a/Engine/Source/Programs/UnrealVersionSelector/Private/Linux/LinuxPlatformInstallation.cpp b/Engine/Source/Programs/UnrealVersionSelector/Private/Linux/LinuxPlatformInstallation.cpp
--- a/Engine/Source/Programs/UnrealVersionSelector/Private/Linux/LinuxPlatformInstallation.cpp
+++ b/Engine/Source/Programs/UnrealVersionSelector/Private/Linux/LinuxPlatformInstallation.cpp
@@ -509,7 +509,7 @@ bool FLinuxPlatformInstallation::LaunchEditor(const FString& RootDirName, const
 	FString FileName;
 	if (ExplicitFileName.IsEmpty())
 	{
-		FileName = RootDirName / TEXT("Engine/Binaries/Linux/UE4Editor");
+		FileName = RootDirName / TEXT("Engine/Binaries/Linux/UE5Editor");
 	}
 	else
 	{
diff --git a/Engine/Source/Programs/UnrealVersionSelector/Private/Windows/WindowsPlatformInstallation.cpp b/Engine/Source/Programs/UnrealVersionSelector/Private/Windows/WindowsPlatformInstallation.cpp
--- a/Engine/Source/Programs/UnrealVersionSelector/Private/Windows/WindowsPlatformInstallation.cpp
+++ b/Engine/Source/Programs/UnrealVersionSelector/Private/Windows/WindowsPlatformInstallation.cpp
@@ -218,7 +218,7 @@ bool FWindowsPlatformInstallation::LaunchEditor(const FString &RootDirName, cons
 	FString FileName;
 	if (ExplicitFileName.IsEmpty())
 	{
-		FileName = RootDirName / TEXT("Engine/Binaries/Win64/UE4Editor.exe");
+		FileName = RootDirName / TEXT("Engine/Binaries/Win64/UE5Editor.exe");
 	}
 	else
 	{
diff --git a/Engine/Source/Runtime/Core/Private/GenericPlatform/GenericPlatformCrashContext.cpp b/Engine/Source/Runtime/Core/Private/GenericPlatform/GenericPlatformCrashContext.cpp
--- a/Engine/Source/Runtime/Core/Private/GenericPlatform/GenericPlatformCrashContext.cpp
+++ b/Engine/Source/Runtime/Core/Private/GenericPlatform/GenericPlatformCrashContext.cpp
@@ -735,14 +735,14 @@ void FGenericCrashContext::AddPortableCallStackHash() const
 	const TCHAR* ExeName = FPlatformProcess::ExecutableName();
 
 	// We dont want this to be thrown into an FString as it will alloc memory
-	const TCHAR* UE4EditorName = TEXT("UE4Editor");
+	const TCHAR* UE4EditorName = TEXT("UE5Editor");
 
 	FSHA1 Sha;
 	FSHAHash Hash;
 
 	for (TArray<FCrashStackFrame>::TConstIterator It(CallStack); It; ++It)
 	{
-		// If we are our own module or our module contains UE4Editor we assume we own these. We cannot depend on offsets of system libs
+		// If we are our own module or our module contains UE5Editor we assume we own these. We cannot depend on offsets of system libs
 		// as they may have different versions
 		if (It->ModuleName == ExeName || It->ModuleName.Contains(UE4EditorName))
 		{
diff --git a/Engine/Source/Runtime/Core/Private/Misc/ConfigCacheIni.cpp b/Engine/Source/Runtime/Core/Private/Misc/ConfigCacheIni.cpp
--- a/Engine/Source/Runtime/Core/Private/Misc/ConfigCacheIni.cpp
+++ b/Engine/Source/Runtime/Core/Private/Misc/ConfigCacheIni.cpp
@@ -3859,7 +3859,7 @@ void FConfigCacheIni::InitializeConfigSystem()
 	// create GConfig
 	GConfig = new FConfigCacheIni(EConfigCacheType::DiskBacked);
 
-	// load the main .ini files (unless we're running a program or a gameless UE4Editor.exe, DefaultEngine.ini is required).
+	// load the main .ini files (unless we're running a program or a gameless UE5Editor.exe, DefaultEngine.ini is required).
 	const bool bIsGamelessExe = !FApp::HasProjectName();
 	const bool bDefaultEngineIniRequired = !bIsGamelessExe && (GIsGameAgnosticExe || FApp::IsProjectNameEmpty());
 	bool bEngineConfigCreated = FConfigCacheIni::LoadGlobalIniFile(GEngineIni, TEXT("Engine"), nullptr, bDefaultEngineIniRequired);
diff --git a/Engine/Source/Runtime/CoreUObject/Public/Misc/HotReloadInterface.h b/Engine/Source/Runtime/CoreUObject/Public/Misc/HotReloadInterface.h
--- a/Engine/Source/Runtime/CoreUObject/Public/Misc/HotReloadInterface.h
+++ b/Engine/Source/Runtime/CoreUObject/Public/Misc/HotReloadInterface.h
@@ -27,7 +27,7 @@ enum class ERecompileModuleFlags : uint8
 	// Report failure if UHT-generated code changes as a result of the recompile
 	FailIfGeneratedCodeChanges = 0x02,
 
-	// Even if this is not code-based project compile with game project as the target for UBT (do not use UE4Editor target)
+	// Even if this is not code-based project compile with game project as the target for UBT (do not use UE5Editor target)
 	ForceCodeProject = 0x04,
 };
 ENUM_CLASS_FLAGS(ERecompileModuleFlags)
diff --git a/Engine/Source/Runtime/Launch/Private/LaunchEngineLoop.cpp b/Engine/Source/Runtime/Launch/Private/LaunchEngineLoop.cpp
--- a/Engine/Source/Runtime/Launch/Private/LaunchEngineLoop.cpp
+++ b/Engine/Source/Runtime/Launch/Private/LaunchEngineLoop.cpp
@@ -1698,7 +1698,7 @@ int32 FEngineLoop::PreInitPreStartupScreen(const TCHAR* CmdLine)
 			// We did not find a non-suffixed folder and we DID find the suffixed one.
 			// The engine MUST be launched with <GameName>Game.
 			const FText GameNameText = FText::FromString(FApp::GetProjectName());
-			FMessageDialog::Open(EAppMsgType::Ok, FText::Format(LOCTEXT("RequiresGamePrefix", "Error: UE4Editor does not append 'Game' to the passed in game name.\nYou must use the full name.\nYou specified '{0}', use '{0}Game'."), GameNameText));
+			FMessageDialog::Open(EAppMsgType::Ok, FText::Format(LOCTEXT("RequiresGamePrefix", "Error: UE5Editor does not append 'Game' to the passed in game name.\nYou must use the full name.\nYou specified '{0}', use '{0}Game'."), GameNameText));
 			return 1;
 		}
 	}
@@ -1810,9 +1810,9 @@ int32 FEngineLoop::PreInitPreStartupScreen(const TCHAR* CmdLine)
 	bool bIsNotEditor = false;
 
 	// This isn't necessarily pretty, but many requests have been made to allow
-	//   UE4Editor.exe <GAMENAME> -game <map>
+	//   UE5Editor.exe <GAMENAME> -game <map>
 	// or
-	//   UE4Editor.exe <GAMENAME> -game 127.0.0.0
+	//   UE5Editor.exe <GAMENAME> -game 127.0.0.0
 	// We don't want to remove the -game from the commandline just yet in case
 	// we need it for something later. So, just move it to the end for now...
 	const bool bFirstTokenIsGame = (Token == TEXT("-GAME"));
@@ -5257,7 +5257,7 @@ bool LaunchCorrectEditorExecutable(const FString& EditorTargetFileName)
 	FString LaunchExecutableName;
 	if(EditorTargetFileName.Len() == 0)
 	{
-		LaunchExecutableName = FPlatformProcess::GenerateApplicationPath(TEXT("UE4Editor"), FApp::GetBuildConfiguration());
+		LaunchExecutableName = FPlatformProcess::GenerateApplicationPath(TEXT("UE5Editor"), FApp::GetBuildConfiguration());
 	}
 	else
 	{
diff --git a/Engine/Source/UE4Editor.Target.cs b/Engine/Source/UE4Editor.Target.cs
deleted file mode 100644
diff --git a/Engine/Source/UE5Editor.Target.cs b/Engine/Source/UE5Editor.Target.cs
new file mode 100644
--- /dev/null
+++ b/Engine/Source/UE5Editor.Target.cs
@@ -0,0 +1,15 @@
+// Copyright Epic Games, Inc. All Rights Reserved.
+
+using UnrealBuildTool;
+using System.Collections.Generic;
+
+public class UE5EditorTarget : TargetRules
+{
+	public UE5EditorTarget( TargetInfo Target ) : base(Target)
+	{
+		Type = TargetType.Editor;
+		BuildEnvironment = TargetBuildEnvironment.Shared;
+		bBuildAllModules = true;
+		ExtraModuleNames.Add("UE4Game");
+	}
+}
