Removing MobileStarterContent from DDC feature packs.

#rnx

diff --git a/Engine/Build/InstalledEngineFilters.xml b/Engine/Build/InstalledEngineFilters.xml
--- a/Engine/Build/InstalledEngineFilters.xml
+++ b/Engine/Build/InstalledEngineFilters.xml
@@ -87,7 +87,6 @@
 		Templates/TP_PhotoStudioBP/TP_PhotoStudioBP.uproject
 		Templates/TP_CollaborativeBP/TP_CollaborativeBP.uproject
 		Samples/StarterContent/StarterContent.uproject
-		Samples/MobileStarterContent/MobileStarterContent.uproject
 	</Property>
 	<!-- Win64 specific DDC -->
 	<Property Name="ProjectsToBuildDDCWin64" Value="$(ProjectsToBuildDDC)">
