Mering fixes to installed build script from 4.26

diff --git a/Engine/Build/InstalledEngineBuild.xml b/Engine/Build/InstalledEngineBuild.xml
--- a/Engine/Build/InstalledEngineBuild.xml
+++ b/Engine/Build/InstalledEngineBuild.xml
@@ -143,10 +143,8 @@
 
 	<Option Name="CrashReporterAPIURL" DefaultValue="" Description="The URL to use to talk to the CrashReporterClient API."/>
 	<Option Name="CrashReporterAPIKey" DefaultValue="" Description="The API key to use to talk to the CrashReporterClient API."/>
-	<Property Name="CrashReporterCompileArgsWin" Value=""/>
-	<Property Name="CrashReporterCompileArgsWin" Value="-CompilerArguments=&quot;/DCRC_TELEMETRY_URL=$(CrashReporterAPIURL) /DCRC_TELEMETRY_KEY=$(CrashReporterAPIKey)&quot;" If="'$(CrashReporterAPIURL)' != '' And '$(CrashReporterAPIKey)' != ''" />
-	<Property Name="CrashReporterCompileArgsMac" Value=""/>
-	<Property Name="CrashReporterCompileArgsMac" Value="-CompilerArguments=&quot; -DCRC_TELEMETRY_URL=$(CrashReporterAPIURL) -DCRC_TELEMETRY_KEY=$(CrashReporterAPIKey)&quot;" If="'$(CrashReporterAPIURL)' != '' And '$(CrashReporterAPIKey)' != ''" />
+	<Property Name="CrashReporterCompileArgs" Value=""/>
+	<Property Name="CrashReporterCompileArgs" Value="-define:CRC_TELEMETRY_URL=\&quot;$(CrashReporterAPIURL)\&quot; -define:CRC_TELEMETRY_KEY_DEV=\&quot;$(CrashReporterAPIKey)\&quot; -define:CRC_TELEMETRY_KEY_RELEASE=\&quot;$(CrashReporterAPIKey)\&quot; -OverrideBuildEnvironment" If="'$(CrashReporterAPIURL)' != '' And '$(CrashReporterAPIKey)' != ''" />
 
 	<!-- Base editor and tools for compile on Windows -->
 	<Agent Name="Editor Win64" Type="Win64_Licensee">
@@ -158,8 +156,8 @@
 			<Compile Target="UnrealHeaderTool" Platform="Win64" Configuration="Development" Arguments="-precompile -allmodules"/>
 		</Node>
 		<Node Name="Compile UnrealEditor Win64" Requires="Compile UnrealHeaderTool Win64" Produces="#UnrealEditor Win64;#UnrealEditor Win64 Unstripped;#UnrealEditor Win64 Stripped;#UnrealEditor Win64 Unsigned;#UnrealEditor Win64 Signed;#UnrealEditor Win64 WithExceptions">
-			<Compile Target="UnrealEditor" Platform="Win64" Configuration="DebugGame" Tag="#UnrealEditor Win64" Arguments="-precompile -allmodules $(CrashReporterCompileArgsWin)"/>
-			<Compile Target="UnrealEditor" Platform="Win64" Configuration="Development" Tag="#UnrealEditor Win64" Arguments="-precompile -allmodules $(CrashReporterCompileArgsWin)"/>
+			<Compile Target="UnrealEditor" Platform="Win64" Configuration="DebugGame" Tag="#UnrealEditor Win64" Arguments="-precompile -allmodules $(CrashReporterCompileArgs)"/>
+			<Compile Target="UnrealEditor" Platform="Win64" Configuration="Development" Tag="#UnrealEditor Win64" Arguments="-precompile -allmodules $(CrashReporterCompileArgs)"/>
 
 			<Do If="$(EmbedSrcSrvInfo)">
 				<!-- Embed source info into the PDB files. Should be done from this machine to ensure that paths are correct. -->
@@ -200,8 +198,8 @@
 			<Compile Target="UnrealHeaderTool" Platform="Mac" Configuration="Development" Arguments="-precompile -allmodules"/>
 		</Node>
 		<Node Name="Compile UnrealEditor Mac" Requires="Compile UnrealHeaderTool Mac" Produces="#UnrealEditor Mac;#UnrealEditor Mac Unstripped;#UnrealEditor Mac Stripped;#UnrealEditor Mac Unsigned;#UnrealEditor Mac Signed;#UnrealEditor Mac WithExceptions">
-			<Compile Target="UnrealEditor" Platform="Mac" Configuration="DebugGame" Tag="#UnrealEditor Mac" Arguments="-precompile -allmodules $(CrashReporterCompileArgsWin)"/>
-			<Compile Target="UnrealEditor" Platform="Mac" Configuration="Development" Tag="#UnrealEditor Mac" Arguments="-precompile -allmodules $(CrashReporterCompileArgsWin)"/>
+			<Compile Target="UnrealEditor" Platform="Mac" Configuration="DebugGame" Tag="#UnrealEditor Mac" Arguments="-precompile -allmodules $(CrashReporterCompileArgs)"/>
+			<Compile Target="UnrealEditor" Platform="Mac" Configuration="Development" Tag="#UnrealEditor Mac" Arguments="-precompile -allmodules $(CrashReporterCompileArgs)"/>
 
 			<Property Name="UnrealEditorExceptions" Value=""/>
 			<Property Name="UnrealEditorExceptions" Value="$(PluginsExceptions)" If="$(CompileDatasmithPlugins)"/>
@@ -450,8 +448,8 @@
 
 	<Agent Name="Tools Group Win64" Type="Win64_Licensee">
 		<Node Name="Build Tools Win64" Requires="Compile UnrealHeaderTool Win64;#NotForLicensee Build Tools Win64" Produces="#Tools Win64 Complete">
-			<Compile Target="CrashReportClient" Configuration="Shipping" Platform="Win64" Tag="#Build Tools Win64" Arguments="$(CrashReporterCompileArgsWin)"/>
-			<Compile Target="CrashReportClientEditor" Configuration="Shipping" Platform="Win64" Tag="#Build Tools Win64" Arguments="$(CrashReporterCompileArgsWin)"/>
+			<Compile Target="CrashReportClient" Configuration="Shipping" Platform="Win64" Tag="#Build Tools Win64" Arguments="$(CrashReporterCompileArgs)"/>
+			<Compile Target="CrashReportClientEditor" Configuration="Shipping" Platform="Win64" Tag="#Build Tools Win64" Arguments="$(CrashReporterCompileArgs)"/>
 			<Tag Files="#NotForLicensee Build Tools Win64" With="#Build Tools Win64"/>
 			<Compile Target="ShaderCompileWorker" Configuration="Development" Platform="Win64" Tag="#Build Tools Win64"/>
 			<Compile Target="UnrealCEFSubProcess" Configuration="Development" Platform="Win64" Tag="#Build Tools Win64"/>
@@ -545,8 +543,8 @@
 
 	<Agent Name="Tools Group Mac" Type="Mac_Licensee">
 		<Node Name="Build Tools Mac" Requires="Compile UnrealHeaderTool Mac">
-			<Compile Target="CrashReportClient" Configuration="Shipping" Platform="Mac" Tag="#Build Tools Mac" Arguments="$(CrashReporterCompileArgsMac)"/>
-			<Compile Target="CrashReportClientEditor" Configuration="Shipping" Platform="Mac" Tag="#Build Tools Mac" Arguments="$(CrashReporterCompileArgsMac)"/>
+			<Compile Target="CrashReportClient" Configuration="Shipping" Platform="Mac" Tag="#Build Tools Mac" Arguments="$(CrashReporterCompileArgs)"/>
+			<Compile Target="CrashReportClientEditor" Configuration="Shipping" Platform="Mac" Tag="#Build Tools Mac" Arguments="$(CrashReporterCompileArgs)"/>
 			<Compile Target="ShaderCompileWorker" Configuration="Development" Platform="Mac" Tag="#Build Tools Mac"/>
 			<Compile Target="UnrealCEFSubProcess" Configuration="Development" Platform="Mac" Tag="#Build Tools Mac"/>
 			<Compile Target="UnrealInsights" Configuration="Development" Platform="Mac" Tag="#Build Tools Mac"/>
