Removed stale references to StaticMeshEditorExtension plugin in Build config files.

diff --git a/Engine/Build/InstalledEngineBuild.xml b/Engine/Build/InstalledEngineBuild.xml
--- a/Engine/Build/InstalledEngineBuild.xml
+++ b/Engine/Build/InstalledEngineBuild.xml
@@ -739,7 +739,7 @@
 	<Agent Name="Extra Plugins Agent Win64" Type="Win64">
 		<Node Name="Compile Datasmith Plugins Win64" Requires="Compile UnrealHeaderTool Win64" Produces="#DatasmithPluginBinaries Win64">
 			<Do If="$(CompileDatasmithPlugins) == 'True'">
-				<Property Name="DatasmithPlugins" Value="DatasmithCADImporter;DatasmithIFCImporter;DatasmithC4DImporter;AxFImporter;StaticMeshEditorExtension;MDLImporter"/>
+				<Property Name="DatasmithPlugins" Value="DatasmithCADImporter;DatasmithIFCImporter;DatasmithC4DImporter;AxFImporter;MDLImporter"/>
 				<ForEach Name="Plugin" Values="$(DatasmithPlugins)">
 					<Compile Target="UnrealEditor" Configuration="Development" Platform="Win64" Arguments="-Plugin=&quot;$(RootDir)/Engine/Plugins/Enterprise/$(Plugin)/$(Plugin).uplugin&quot;"/>
 					<Tag Files="$(RootDir)/Engine/Plugins/Enterprise/$(Plugin)/Binaries/..." With="#DatasmithPluginBinaries Win64"/>
@@ -751,7 +751,7 @@
 	<Agent Name="Extra Plugins Agent Mac" Type="Mac">
 		<Node Name="Compile Datasmith Plugins Mac" Requires="Compile UnrealHeaderTool Mac" Produces="#DatasmithPluginBinaries Mac">
 			<Do If="$(CompileDatasmithPlugins) == 'True'">
-				<Property Name="DatasmithPlugins" Value="DatasmithCADImporter;DatasmithIFCImporter;DatasmithC4DImporter;AxFImporter;StaticMeshEditorExtension;MDLImporter"/>
+				<Property Name="DatasmithPlugins" Value="DatasmithCADImporter;DatasmithIFCImporter;DatasmithC4DImporter;AxFImporter;MDLImporter"/>
 				<ForEach Name="Plugin" Values="$(DatasmithPlugins)">
 					<Compile Target="UnrealEditor" Configuration="Development" Platform="Mac" Arguments="-Plugin=&quot;$(RootDir)/Engine/Plugins/Enterprise/$(Plugin)/$(Plugin).uplugin&quot;"/>
 					<Tag Files="$(RootDir)/Engine/Plugins/Enterprise/$(Plugin)/Binaries/..." With="#DatasmithPluginBinaries Mac"/>
diff --git a/Engine/Build/InstalledEngineFilters.xml b/Engine/Build/InstalledEngineFilters.xml
--- a/Engine/Build/InstalledEngineFilters.xml
+++ b/Engine/Build/InstalledEngineFilters.xml
@@ -32,7 +32,6 @@
 		Engine/Plugins/Enterprise/DatasmithIFCImporter/...
 		Engine/Plugins/Enterprise/DatasmithC4DImporter/...
 		Engine/Plugins/Enterprise/AxFImporter/...
-		Engine/Plugins/Enterprise/StaticMeshEditorExtension/...
 		Engine/Plugins/Enterprise/MDLImporter/...
 	</Property>
 	<Property Name="MacSignFilter" Value="*.dylib;*.app"/>
