Fix Mac Datasmith plugins getting added when compiling Win64 only.

#rnx

diff --git a/Engine/Build/InstalledEngineBuild.xml b/Engine/Build/InstalledEngineBuild.xml
--- a/Engine/Build/InstalledEngineBuild.xml
+++ b/Engine/Build/InstalledEngineBuild.xml
@@ -738,7 +738,7 @@
 	<!-- These are plugins that need to be compiled with a full workspace rather than Win64_Licensee. -->
 	<Agent Name="Extra Plugins Agent Win64" Type="Win64">
 		<Node Name="Compile Datasmith Plugins Win64" Requires="Compile UnrealHeaderTool Win64" Produces="#DatasmithPluginBinaries Win64">
-			<Do If="$(CompileDatasmithPlugins) == 'True'">
+			<Do If="$(WithWin64) And $(CompileDatasmithPlugins) == 'True'">
 				<Property Name="DatasmithPlugins" Value="DatasmithCADImporter;DatasmithIFCImporter;DatasmithC4DImporter;AxFImporter;MDLImporter"/>
 				<ForEach Name="Plugin" Values="$(DatasmithPlugins)">
 					<Compile Target="UnrealEditor" Configuration="Development" Platform="Win64" Arguments="-Plugin=&quot;$(RootDir)/Engine/Plugins/Enterprise/$(Plugin)/$(Plugin).uplugin&quot;"/>
@@ -750,7 +750,7 @@
 
 	<Agent Name="Extra Plugins Agent Mac" Type="Mac">
 		<Node Name="Compile Datasmith Plugins Mac" Requires="Compile UnrealHeaderTool Mac" Produces="#DatasmithPluginBinaries Mac">
-			<Do If="$(CompileDatasmithPlugins) == 'True'">
+			<Do If="$(WithMac) And $(CompileDatasmithPlugins) == 'True'">
 				<Property Name="DatasmithPlugins" Value="DatasmithCADImporter;DatasmithIFCImporter;DatasmithC4DImporter;AxFImporter;MDLImporter"/>
 				<ForEach Name="Plugin" Values="$(DatasmithPlugins)">
 					<Compile Target="UnrealEditor" Configuration="Development" Platform="Mac" Arguments="-Plugin=&quot;$(RootDir)/Engine/Plugins/Enterprise/$(Plugin)/$(Plugin).uplugin&quot;"/>
@@ -778,7 +778,7 @@
 		<Property Name="InstalledRequirements" Value="$(InstalledRequirements);#UnrealGame Lumin" If="'$(WithLumin)' == true"/>
 		<Property Name="InstalledRequirements" Value="$(InstalledRequirements);#UnrealGame HoloLens" If="'$(WithHoloLens)' == true"/>
 		<Property Name="InstalledRequirements" Value="$(InstalledRequirements);#CompressedDDCWin64" If="'$(WithDDC)' == true"/>
-		<Property Name="InstalledRequirements" Value="$(InstalledRequirements);#Compile Datasmith Plugins Win64;#DatasmithPluginBinaries Win64" If="$(CompileDatasmithPlugins)"/>
+		<Property Name="InstalledRequirements" Value="$(InstalledRequirements);#Compile Datasmith Plugins Win64;#DatasmithPluginBinaries Win64" If="$(WithWin64) And $(CompileDatasmithPlugins)"/>
 		<!-- Feature packs -->
 		<Property Name="InstalledRequirements" Value="$(InstalledRequirements);Make Feature Packs"/>
 
@@ -936,7 +936,7 @@
 		<Property Name="FilterRequirements" Value="$(FilterRequirements);#UnrealGame TVOS;#UnrealGame TVOS Unstripped;#UnrealGame TVOS Stripped" If="'$(WithTVOS)' == true"/>
 		<Property Name="FilterRequirements" Value="$(FilterRequirements);#UnrealGame LuminMac;#UnrealGame LuminMac Unstripped;#UnrealGame LuminMac Stripped" If="'$(WithLuminMac)' == true"/>
 		<Property Name="FilterRequirements" Value="$(FilterRequirements);#CompressedDDCMac" If="'$(WithDDC)' == true"/>
-		<Property Name="FilterRequirements" Value="$(FilterRequirements);#Compile Datasmith Plugins Mac;#DatasmithPluginBinaries Mac" If="$(CompileDatasmithPlugins)"/>
+		<Property Name="FilterRequirements" Value="$(FilterRequirements);#Compile Datasmith Plugins Mac;#DatasmithPluginBinaries Mac" If="$(WithMac) And $(CompileDatasmithPlugins)"/>
 
 		<!-- Feature packs -->
 		<Property Name="FilterRequirements" Value="$(FilterRequirements);Make Feature Packs"/>
@@ -958,7 +958,7 @@
 			<Tag Files="$(UnrealEditor)" Except="#UnrealEditor Mac Unstripped;#UnrealEditor Mac Unsigned" With="#Installed Mac"/>
 			<Tag Files="#UnrealEditor Mac Stripped;#UnrealEditor Mac Signed" With="#Saved Output"/>
 			<Tag Files="#Build Tools Mac;#Make Feature Packs;#Rules Assemblies" With="#Installed Mac"/>
-			<Tag Files="#DatasmithPluginBinaries Mac" With="#Installed Mac" If="$(CompileDatasmithPlugins)"/>
+			<Tag Files="#DatasmithPluginBinaries Mac" With="#Installed Mac" If="$(WithMac) And $(CompileDatasmithPlugins)"/>
 
 			<!-- Tag enabled target platforms -->
 			<Do If="'$(WithMac)' == true">
