Avoid building our temporary net core csproj when build UAT during creation of installed build.

diff --git a/Engine/Build/InstalledEngineBuild.xml b/Engine/Build/InstalledEngineBuild.xml
--- a/Engine/Build/InstalledEngineBuild.xml
+++ b/Engine/Build/InstalledEngineBuild.xml
@@ -481,7 +481,7 @@
 			<Delete Files="$(CsToolsDir)/..."/>
 
 			<!-- Copy Source and referenced libraries to a new location with Confidential folders removed-->
-			<Tag Files="$(AutomationToolPaths)" Filter="*.csproj" Except="$(ConfidentialExceptions)"  With="#UAT Projects"/>
+			<Tag Files="$(AutomationToolPaths)" Filter="*.csproj" Except="$(ConfidentialExceptions);*Core.csproj"  With="#UAT Projects"/>
 			<CsCompile Project="#UAT Projects" Configuration="Development" Platform="AnyCPU" TagReferences="#UATReferences" EnumerateOnly="true"/>
 			<Tag Files="Engine/Binaries/DotNET/...;Engine/Binaries/ThirdParty/Newtonsoft/...;Engine/Binaries/ThirdParty/VisualStudio/...;Engine/Source/Programs/...;Engine/Platforms/*/Source/Programs/...;Engine/Source/Editor/SwarmInterface/...;#UATReferences" Except="$(ConfidentialExceptions);AutomationTool.csproj.References" With="#RedistUATSource"/>
 			<ForEach Name="Platform" Values="$(ExtensionPlatforms)">
@@ -493,7 +493,7 @@
 			<ForEach Name="AutomationToolPath" Values="$(AutomationToolPaths)">
 				<Property Name="RedistAutomationToolPaths" Value="$(RedistAutomationToolPaths)$(CsToolsDir)/$(AutomationToolPath);"/>
 			</ForEach>
-			<Tag Files="$(RedistAutomationToolPaths)" Filter="*.csproj" With="#RedistUATProjects"/>
+			<Tag Files="$(RedistAutomationToolPaths)" Filter="*.csproj" Except="*Core.csproj" With="#RedistUATProjects"/>
 
 			<!-- Compile all the tools -->
 			<CsCompile Project="#RedistUATProjects" Configuration="Development" Platform="AnyCPU" Tag="#Build Tools CS Binaries;#Sign Binaries" TagReferences="#Build Tools CS Binaries"/>
