A platform's DDC should not be included when WithDDC is set if that paltform is not being built.

#rnx

diff --git a/Engine/Build/InstalledEngineBuild.xml b/Engine/Build/InstalledEngineBuild.xml
--- a/Engine/Build/InstalledEngineBuild.xml
+++ b/Engine/Build/InstalledEngineBuild.xml
@@ -777,7 +777,7 @@
 		<Property Name="InstalledRequirements" Value="$(InstalledRequirements);#UnrealGame LinuxAArch64;#UnrealGame LinuxAArch64 Unstripped;#UnrealGame LinuxAArch64 Stripped" If="'$(WithLinuxAArch64)' == true"/>
 		<Property Name="InstalledRequirements" Value="$(InstalledRequirements);#UnrealGame Lumin" If="'$(WithLumin)' == true"/>
 		<Property Name="InstalledRequirements" Value="$(InstalledRequirements);#UnrealGame HoloLens" If="'$(WithHoloLens)' == true"/>
-		<Property Name="InstalledRequirements" Value="$(InstalledRequirements);#CompressedDDCWin64" If="'$(WithDDC)' == true"/>
+		<Property Name="InstalledRequirements" Value="$(InstalledRequirements);#CompressedDDCWin64" If="$(WithDDC) And $(WithWin64)"/>
 		<Property Name="InstalledRequirements" Value="$(InstalledRequirements);#Compile Datasmith Plugins Win64;#DatasmithPluginBinaries Win64" If="$(WithWin64) And $(CompileDatasmithPlugins)"/>
 		<!-- Feature packs -->
 		<Property Name="InstalledRequirements" Value="$(InstalledRequirements);Make Feature Packs"/>
@@ -935,7 +935,7 @@
 		<Property Name="FilterRequirements" Value="$(FilterRequirements);#UnrealGame IOS;#UnrealGame IOS Unstripped;#UnrealGame IOS Stripped" If="'$(WithIOS)' == true"/>
 		<Property Name="FilterRequirements" Value="$(FilterRequirements);#UnrealGame TVOS;#UnrealGame TVOS Unstripped;#UnrealGame TVOS Stripped" If="'$(WithTVOS)' == true"/>
 		<Property Name="FilterRequirements" Value="$(FilterRequirements);#UnrealGame LuminMac;#UnrealGame LuminMac Unstripped;#UnrealGame LuminMac Stripped" If="'$(WithLuminMac)' == true"/>
-		<Property Name="FilterRequirements" Value="$(FilterRequirements);#CompressedDDCMac" If="'$(WithDDC)' == true"/>
+		<Property Name="FilterRequirements" Value="$(FilterRequirements);#CompressedDDCMac" If="$(WithDDC) And $(WithMac)"/>
 		<Property Name="FilterRequirements" Value="$(FilterRequirements);#Compile Datasmith Plugins Mac;#DatasmithPluginBinaries Mac" If="$(WithMac) And $(CompileDatasmithPlugins)"/>
 
 		<!-- Feature packs -->
@@ -1062,7 +1062,7 @@
 		<!-- Optional Target Platform requirements -->
 		<Property Name="InstalledRequirements" Value="$(InstalledRequirements);#UnrealGame Linux;#UnrealGame Linux Unstripped;#UnrealGame Linux Stripped" If="'$(WithLinux)' == true"/>
 		<Property Name="InstalledRequirements" Value="$(InstalledRequirements);#UnrealGame LinuxAArch64;#UnrealGame LinuxAArch64 Unstripped;#UnrealGame LinuxAArch64 Stripped" If="'$(WithLinuxAArch64)' == true"/>
-		<Property Name="InstalledRequirements" Value="$(InstalledRequirements);#CompressedDDCLinux" If="'$(WithDDC)' == true"/>
+		<Property Name="InstalledRequirements" Value="$(InstalledRequirements);#CompressedDDCLinux" If="$(WithDDC) And $(WithLinux)"/>
 
 		<!-- Feature packs -->
 		<Property Name="InstalledRequirements" Value="$(InstalledRequirements);Make Feature Packs"/>
