Fixing graph break.

#rnx

diff --git a/Engine/Build/InstalledEngineBuild.xml b/Engine/Build/InstalledEngineBuild.xml
--- a/Engine/Build/InstalledEngineBuild.xml
+++ b/Engine/Build/InstalledEngineBuild.xml
@@ -902,7 +902,7 @@
 
 			<!-- Copy the compressed DDC -->
 			<Property Name="SavedDir" Value="$(SavedOutput)\Installed\Win64"/>
-			<Copy Files="#CompressedDDCWin64" From="$(SavedDir)" To="$(LocalInstalledDir)" If="'$(WithDDC)' == true" Tag="#Installed Build Win64 Files"/>
+			<Copy Files="#CompressedDDCWin64" From="$(SavedDir)" To="$(LocalInstalledDir)" If="$(WithDDC) And $(WithWin64)" Tag="#Installed Build Win64 Files"/>
 
 			<!-- List all of the platforms that we only allow for content projects -->
 			<Property Name="ContentOnlyPlatforms" Value=""/>
@@ -1039,7 +1039,7 @@
 
 			<!-- Copy the compressed DDC -->
 			<Property Name="SavedDir" Value="$(SavedOutput)/Installed/Mac"/>
-			<Copy Files="#CompressedDDCMac" From="$(SavedDir)" To="$(LocalInstalledDirMac)" If="'$(WithDDC)' == true" Tag="#Installed Build Mac Files"/>
+			<Copy Files="#CompressedDDCMac" From="$(SavedDir)" To="$(LocalInstalledDirMac)" If="$(WithDDC) And $(WithMac)" Tag="#Installed Build Mac Files"/>
 
 			<!-- Create command line arguments for any platforms that were built with additional architectures -->
 			<Property Name="ArchitectureCommands" Value=""/>
@@ -1133,7 +1133,7 @@
 
 			<!-- Copy the compressed DDC -->
 			<Property Name="SavedDir" Value="$(SavedOutput)\Installed\Linux"/>
-			<Copy Files="#CompressedDDCLinux" From="$(SavedDir)" To="$(LocalInstalledDirLinux)" If="'$(WithDDC)' == true" Tag="#Installed Build Linux Files"/>
+			<Copy Files="#CompressedDDCLinux" From="$(SavedDir)" To="$(LocalInstalledDirLinux)" If="$(WithDDC) And $(WithLinux)" Tag="#Installed Build Linux Files"/>
 
 			<!-- Copy the toolchain if it exists -->
 			<Property Name="LinuxToolchainDir" Value="$(ToolchainDir)/HostLinux"/>
