Attempting to get Mac to stop being built when only Win64 is selected.

#rnx

diff --git a/Engine/Build/InstalledEngineBuild.xml b/Engine/Build/InstalledEngineBuild.xml
--- a/Engine/Build/InstalledEngineBuild.xml
+++ b/Engine/Build/InstalledEngineBuild.xml
@@ -1062,7 +1062,7 @@
 		<!-- Optional Target Platform requirements -->
 		<Property Name="InstalledRequirements" Value="$(InstalledRequirements);#UnrealGame Linux;#UnrealGame Linux Unstripped;#UnrealGame Linux Stripped" If="'$(WithLinux)' == true"/>
 		<Property Name="InstalledRequirements" Value="$(InstalledRequirements);#UnrealGame LinuxAArch64;#UnrealGame LinuxAArch64 Unstripped;#UnrealGame LinuxAArch64 Stripped" If="'$(WithLinuxAArch64)' == true"/>
-		<Property Name="InstalledRequirements" Value="$(InstalledRequirements);#CompressedDDCLinux" If="$(WithDDC) And $(WithLinux)"/>
+		<Property Name="InstalledRequirements" Value="$(InstalledRequirements);#CompressedDDCLinux" If="'$(WithDDC)' == true And $(WithLinux)"/>
 
 		<!-- Feature packs -->
 		<Property Name="InstalledRequirements" Value="$(InstalledRequirements);Make Feature Packs"/>
@@ -1133,7 +1133,7 @@
 
 			<!-- Copy the compressed DDC -->
 			<Property Name="SavedDir" Value="$(SavedOutput)\Installed\Linux"/>
-			<Copy Files="#CompressedDDCLinux" From="$(SavedDir)" To="$(LocalInstalledDirLinux)" If="$(WithDDC) And $(WithLinux)" Tag="#Installed Build Linux Files"/>
+			<Copy Files="#CompressedDDCLinux" From="$(SavedDir)" To="$(LocalInstalledDirLinux)" If="'$(WithDDC)' == true And $(WithLinux)" Tag="#Installed Build Linux Files"/>
 
 			<!-- Copy the toolchain if it exists -->
 			<Property Name="LinuxToolchainDir" Value="$(ToolchainDir)/HostLinux"/>
