- Removed android 32bit support from PhysXBuild.xml

diff --git a/Engine/Build/PhysXBuild.xml b/Engine/Build/PhysXBuild.xml
--- a/Engine/Build/PhysXBuild.xml
+++ b/Engine/Build/PhysXBuild.xml
@@ -69,21 +69,11 @@
 			<Command Name="BuildPhysX" Arguments="-TargetPlatforms=Linux-i686-unknown-linux-gnu -TargetConfigs=&quot;$(TargetConfigs)&quot; -TargetWindowsCompilers=&quot;$(TargetWindowsCompilers)&quot; $(OtherArgs)"/>
 		</Node>
 	</Agent>
-	<Agent Name="PhysX Linux ARM32" Type="EngineWin64;PhysX_CompileWin64;Win64">
-		<Node Name="Compile PhysX Linux ARM32">
-			<Command Name="BuildPhysX" Arguments="-TargetPlatforms=Linux-arm-unknown-linux-gnueabihf -TargetConfigs=&quot;$(TargetConfigs)&quot; -TargetWindowsCompilers=&quot;$(TargetWindowsCompilers)&quot; $(OtherArgs)"/>
-		</Node>
-	</Agent>
 	<Agent Name="PhysX Linux ARM64" Type="EngineWin64;PhysX_CompileWin64;Win64">
 		<Node Name="Compile PhysX Linux ARM64">
 			<Command Name="BuildPhysX" Arguments="-TargetPlatforms=Linux-aarch64-unknown-linux-gnueabi -TargetConfigs=&quot;$(TargetConfigs)&quot; -TargetWindowsCompilers=&quot;$(TargetWindowsCompilers)&quot; $(OtherArgs)"/>
 		</Node>
 	</Agent>
-	<Agent Name="PhysX Android Armv7" Type="EngineWin64;PhysX_CompileWin64;Win64">
-		<Node Name="Compile PhysX Android Armv7">
-			<Command Name="BuildPhysX" Arguments="-TargetPlatforms=Android-armv7 -TargetConfigs=&quot;$(TargetConfigs)&quot; -TargetWindowsCompilers=&quot;$(TargetWindowsCompilers)&quot; $(OtherArgs)"/>
-		</Node>
-	</Agent>
 	<Agent Name="PhysX Android Arm64" Type="EngineWin64;PhysX_CompileWin64;Win64">
 		<Node Name="Compile PhysX Android Arm64">
 			<Command Name="BuildPhysX" Arguments="-TargetPlatforms=Android-arm64 -TargetConfigs=&quot;$(TargetConfigs)&quot; -TargetWindowsCompilers=&quot;$(TargetWindowsCompilers)&quot; $(OtherArgs)"/>
@@ -129,9 +119,7 @@
 	<Property Name="PhysXTargets" Value="$(PhysXTargets);Compile PhysX Linux x86" If="$(WithLinux-x86)"/>
 	<Property Name="PhysXTargets" Value="$(PhysXTargets);Compile PhysX Linux ARM32" If="$(WithLinux-ARM32)"/>
 	<Property Name="PhysXTargets" Value="$(PhysXTargets);Compile PhysX Linux ARM64" If="$(WithLinux-ARM64)"/>
-	<Property Name="PhysXTargets" Value="$(PhysXTargets);Compile PhysX Android Armv7" If="$(WithAndroid-ARMv7)"/>
 	<Property Name="PhysXTargets" Value="$(PhysXTargets);Compile PhysX Android Arm64" If="$(WithAndroid-ARM64)"/>
-	<Property Name="PhysXTargets" Value="$(PhysXTargets);Compile PhysX Android x86" If="$(WithAndroid-x86)"/>
 	<Property Name="PhysXTargets" Value="$(PhysXTargets);Compile PhysX Android x64" If="$(WithAndroid-x64)"/>
 	<Property Name="PhysXTargets" Value="$(PhysXTargets);Compile PhysX Mac" If="$(WithMac)"/>
 	<Property Name="PhysXTargets" Value="$(PhysXTargets);Compile PhysX IOS" If="$(WithIOS)"/>
