Fixing graph break when Win64 isn't targeted for a release build.

#rtnx

diff --git a/Engine/Build/InstalledEngineBuild.xml b/Engine/Build/InstalledEngineBuild.xml
--- a/Engine/Build/InstalledEngineBuild.xml
+++ b/Engine/Build/InstalledEngineBuild.xml
@@ -797,7 +797,7 @@
 			<Tag Files="$(UnrealEditor)" Except="#UnrealEditor Win64 Unstripped;#UnrealEditor Win64 Unsigned" With="#Installed Win64"/>
 			<Tag Files="#UnrealEditor Win64 Stripped;#UnrealEditor Win64 Signed" With="#Saved Output"/>
 			<Tag Files="#Build Tools Win64;#Make Feature Packs;#Rules Assemblies" With="#Installed Win64"/>
-			<Tag Files="#DatasmithPluginBinaries Win64" With="#Installed Win64" If="$(CompileDatasmithPlugins)"/>
+			<Tag Files="#DatasmithPluginBinaries Win64" With="#Installed Win64" If="$(WithWin64) And $(CompileDatasmithPlugins)"/>
 
 			<!-- Tag enabled target platforms -->
 			<Do If="'$(WithWin64)' == true">
