Add missing files to Android filter

#ue5
#android

diff --git a/Engine/Build/InstalledEngineFilters.xml b/Engine/Build/InstalledEngineFilters.xml
--- a/Engine/Build/InstalledEngineFilters.xml
+++ b/Engine/Build/InstalledEngineFilters.xml
@@ -415,6 +415,7 @@
 		Engine/Source/ThirdParty/GoogleGameSDK/...
 		Engine/Source/ThirdParty/GoogleVR/...
 		Engine/Source/ThirdParty/GoogleInstantPreview/InstantPreview.apk
+		Engine/Source/ThirdParty/HWCPipe/...
 		Engine/Source/ThirdParty/Oculus/LibOVRAudio/LibOVRAudio/lib/*/*.so
 		Engine/Source/ThirdParty/Oculus/LibOVRAudio/*.tps
 		Engine/Source/ThirdParty/Oculus/LibOVRAudio/*.xml
