Add missing DotNet, libimobiledevice and shell scripts to Launcher build.

diff --git a/Engine/Build/InstalledEngineFilters.xml b/Engine/Build/InstalledEngineFilters.xml
--- a/Engine/Build/InstalledEngineFilters.xml
+++ b/Engine/Build/InstalledEngineFilters.xml
@@ -315,6 +315,7 @@
 		Engine/Build/BatchFiles/Mac/GenerateProjectFiles.sh
 		Engine/Build/BatchFiles/Mac/SetupEnvironment.sh
 		Engine/Build/BatchFiles/Mac/SetupMono.sh
+		Engine/Build/BatchFiles/Mac/SetupDotnet.sh
 		Engine/Build/BatchFiles/Mac/FixMonoFiles.sh
 		Engine/Build/BatchFiles/Mac/RunMono.sh
 		Engine/Build/BatchFiles/Mac/RunXBuild.sh
@@ -324,6 +325,10 @@
 
 		<!-- Mono -->
 		Engine/Binaries/ThirdParty/Mono/Mac/...
+
+		<!-- DotNet -->
+		/Binaries/ThirdParty/DotNet/Mac/...
+		/Binaries/ThirdParty/VisualStudio/...
 	</Property>
 
 	<Property Name="CopyEditorFilterMac" Value="$(CopyEditorFilter);$(MacBuildFiles)">
@@ -400,6 +405,7 @@
 	<!-- Mac -->
 	<Property Name="CopyMacFilter">
 		Engine/Binaries/Mac/BootstrapPackagedGame.app/...
+		Engine/Extras/ThirdPartyNotUE/libimobiledevice/Mac/...
 	</Property>
 	<Property Name="CopyMacExceptions">
 	</Property>
@@ -452,6 +458,7 @@
 		Engine/Build/BatchFiles/MakeAndInstallSSHKey.bat
 		Engine/Build/Rsync/...
 		Engine/Extras/ThirdPartyNotUE/DeltaCopy/...
+		Engine/Extras/ThirdPartyNotUE/libimobiledevice/x64/...
 	</Property>
 
 	<Property Name="CopyIOSExceptions">
@@ -473,6 +480,7 @@
 		Engine/Binaries/ThirdParty/IOS/*
 		Engine/Build/BatchFiles/MakeAndInstallSSHKey.bat
 		Engine/Extras/ThirdPartyNotUE/DeltaCopy/...
+		Engine/Extras/ThirdPartyNotUE/libimobiledevice/x64/...
 	</Property>
 
 	<Property Name="CopyTVOSExceptions">
