Install 2-Byte char lldb data formatter for Mac.

#rnx

diff --git a/Engine/Build/InstalledEngineFilters.xml b/Engine/Build/InstalledEngineFilters.xml
--- a/Engine/Build/InstalledEngineFilters.xml
+++ b/Engine/Build/InstalledEngineFilters.xml
@@ -333,7 +333,7 @@
 
 	<Property Name="CopyEditorFilterMac" Value="$(CopyEditorFilter);$(MacBuildFiles)">
 		<!-- Extra tools -->
-		Engine/Extras/LLDBDataFormatters/UEDataFormatters.py
+		Engine/Extras/LLDBDataFormatters/UEDataFormatters_2ByteChars.py
 	</Property>
 
 	<Property Name="CopyEditorExceptionsMac" Value="$(CopyEditorExceptions)">
