Testing installed build mac uat/ubt changes

diff --git a/Engine/Build/InstalledEngineBuild.xml b/Engine/Build/InstalledEngineBuild.xml
--- a/Engine/Build/InstalledEngineBuild.xml
+++ b/Engine/Build/InstalledEngineBuild.xml
@@ -543,6 +543,44 @@
 	</Agent>
 
 	<Agent Name="Tools Group Mac" Type="Mac_Licensee">
+	
+		<!-- Compile the C# tools. Copies files to a temporary directory first, to strip out all confidential folders and avoid sharing violations with the running version of UAT. -->
+		<Node Name="Build Tools CS Mac" Requires="Update Version Files" Produces="#Build Tools CS Binaries Mac">
+		
+			<Delete Files="$(CsToolsDir)/..."/>
+
+			<!-- Copy Source and referenced libraries to a new location with Confidential folders removed-->
+			<Tag Files="$(AutomationToolPaths)" Filter="*.csproj" Except="$(ConfidentialExceptions)"  With="#UAT Projects Mac"/>
+			<CsCompile Project="#UAT Projects Mac" Configuration="Development" Platform="AnyCPU" TagReferences="#UATReferencesMac" EnumerateOnly="true"/>
+			<Tag Files="Engine/Binaries/DotNET/...;Engine/Binaries/ThirdParty/Newtonsoft/...;Engine/Binaries/ThirdParty/VisualStudio/...;Engine/Source/Programs/...;Engine/Platforms/*/Source/Programs/...;Engine/Source/Editor/SwarmInterface/...;#UATReferencesMac" Except="$(ConfidentialExceptions);AutomationTool.csproj.References" With="#RedistUATSourceMac"/>
+			<ForEach Name="Platform" Values="$(ExtensionPlatforms)">
+				<Expand Name="$(Platform)_TagCsToolsRedistUATSource"/>
+			</ForEach>
+			<Copy Files="#RedistUATSourceMac" From="$(RootDir)" To="$(CsToolsDir)"/>
+
+			<Property Name="RedistAutomationToolPaths" Value=""/>
+			<ForEach Name="AutomationToolPath" Values="$(AutomationToolPaths)">
+				<Property Name="RedistAutomationToolPaths" Value="$(RedistAutomationToolPaths)$(CsToolsDir)/$(AutomationToolPath);"/>
+			</ForEach>
+			<Tag Files="$(RedistAutomationToolPaths)" Filter="*.csproj" With="#RedistUATProjectsMac"/>
+
+			<!-- Compile all the tools -->
+			<CsCompile Project="#RedistUATProjectsMac" Configuration="Development" Platform="AnyCPU" Tag="#Build Tools CS Binaries Mac;#Sign Binaries Mac" TagReferences="#Build Tools CS Binaries Mac"/>
+			
+			<!-- Tag AutomationTool and UnrealBuildTool folders recursively as NET Core dependencies are not currently handled by CsCompile -->
+			<Tag Files="$(CsToolsDir)/Engine/Binaries/DotNET/AutomationTool/..." With="#Build Tools CS Binaries Mac" />
+			<Tag Files="$(CsToolsDir)/Engine/Binaries/DotNET/UnrealBuildTool/..." With="#Build Tools CS Binaries Mac" />
+						
+			<CsCompile Project="$(CsToolsDir)/Engine/Source/Programs/NetworkProfiler/NetworkProfiler/NetworkProfiler.csproj" Configuration="Development" Platform="AnyCPU" Tag="#Build Tools CS Binaries Mac" TagReferences="#Build Tools CS Binaries Mac;#Sign Binaries Mac"/>
+			<CsCompile Project="$(CsToolsDir)/Engine/Source/Programs/IOS/DeploymentServer/DeploymentServer.csproj" Configuration="Development" Platform="AnyCPU" Arguments="/verbosity:minimal /target:Rebuild" Tag="#Build Tools CS Binaries Mac;#Sign Binaries Mac" TagReferences="#Build Tools CS Binaries Mac"/>
+			<CsCompile Project="$(CsToolsDir)/Engine/Source/Programs/IOS/MobileDeviceInterface/MobileDeviceInterface.csproj" Configuration="Development" Platform="AnyCPU" Arguments="/verbosity:minimal /target:Rebuild" Tag="#Build Tools CS Binaries Mac;#Sign Binaries Mac" TagReferences="#Build Tools CS Binaries Mac"/>
+			<CsCompile Project="$(CsToolsDir)/Engine/Source/Programs/IOS/DeploymentInterface/DeploymentInterface.csproj" Configuration="Development" Platform="AnyCPU" Arguments="/verbosity:minimal /target:Rebuild" Tag="#Build Tools CS Binaries Mac;#Sign Binaries Mac" TagReferences="#Build Tools CS Binaries Mac"/>
+			<CsCompile Project="$(CsToolsDir)/Engine/Source/Programs/IOS/DeploymentServerLauncher/DeploymentServerLauncher.csproj" Configuration="Development" Platform="AnyCPU" Arguments="/verbosity:minimal /target:Rebuild" Tag="#Build Tools CS Binaries Mac;#Sign Binaries Mac" TagReferences="#Build Tools CS Binaries Mac"/>
+
+			<!-- Sign the binaries -->
+			<Sign Files="#Sign Binaries Mac" If="'$(SignExecutables)' == true"/>
+		</Node>
+	
 		<Node Name="Build Tools Mac" Requires="Compile UnrealHeaderTool Mac">
 			<Compile Target="CrashReportClient" Configuration="Shipping" Platform="Mac" Tag="#Build Tools Mac" Arguments="$(CrashReporterCompileArgs)"/>
 			<Compile Target="CrashReportClientEditor" Configuration="Shipping" Platform="Mac" Tag="#Build Tools Mac" Arguments="$(CrashReporterCompileArgs)"/>
@@ -928,7 +966,7 @@
 		<!-- Build up a list of requirements as some can be optional -->
 		<Property Name="UnrealEditor" Value="#UnrealEditor Mac" />
 		<Property Name="UnrealEditor" Value="#UnrealEditor Mac WithExceptions" If="$(CompileDatasmithPlugins)" />
-		<Property Name="FilterRequirements" Value="#Compile UnrealHeaderTool Mac;$(UnrealEditor);#UnrealEditor Mac Unstripped;#UnrealEditor Mac Stripped;#UnrealEditor Mac Unsigned;#UnrealEditor Mac Signed;#Build Tools Mac;#Build Tools CS Binaries;#Rules Assemblies"/>
+		<Property Name="FilterRequirements" Value="#Compile UnrealHeaderTool Mac;$(UnrealEditor);#UnrealEditor Mac Unstripped;#UnrealEditor Mac Stripped;#UnrealEditor Mac Unsigned;#UnrealEditor Mac Signed;#Build Tools Mac;#Build Tools CS Binaries Mac;#Rules Assemblies"/>
 
 		<!-- Optional Target Platform requirements -->
 		<Property Name="FilterRequirements" Value="$(FilterRequirements);#UnrealGame Mac;#UnrealGame Mac Unstripped;#UnrealGame Mac Stripped;#UnrealGame Mac Unsigned;#UnrealGame Mac Signed" If="'$(WithMac)' == true"/>
@@ -1023,7 +1061,7 @@
 			<!-- Tag files that need to be copied directly -->
 			<Tag Files="#Installed Mac" Except="$(CopyInstalledExceptions);#Strip Tools for Installed Mac;#Sign Tools for Installed Mac" With="#Copy for Installed Mac"/>
 			<Copy Files="#Copy for Installed Mac" From="$(RootDir)" To="$(LocalInstalledDirMac)" Tag="#Installed Build Mac Files"/>
-			<Copy Files="#Build Tools CS Binaries" From="$(CsToolsDir)" To="$(LocalInstalledDirMac)" Tag="#Installed Build Mac Files"/>
+			<Copy Files="#Build Tools CS Binaries Mac" From="$(CsToolsDir)" To="$(LocalInstalledDirMac)" Tag="#Installed Build Mac Files"/>
 			
 			<!-- Tag saved output that needs to be copied directly -->
 			<Tag BaseDir="$(SavedOutput)" Files="#Saved Output" Except="$(CopyInstalledExceptions)" With="#Copy Saved Output"/>
