Fixed typo in filter path for DotNet in Installed builds
Fixed remote toolchain thinking it didn't need to compile UBT/UAT on remote Mac when packaging IOS C++ from installed build.

diff --git a/Engine/Build/InstalledEngineFilters.xml b/Engine/Build/InstalledEngineFilters.xml
--- a/Engine/Build/InstalledEngineFilters.xml
+++ b/Engine/Build/InstalledEngineFilters.xml
@@ -327,8 +327,8 @@
 		Engine/Binaries/ThirdParty/Mono/Mac/...
 
 		<!-- DotNet -->
-		/Binaries/ThirdParty/DotNet/Mac/...
-		/Binaries/ThirdParty/VisualStudio/...
+		Engine/Binaries/ThirdParty/DotNet/Mac/...
+		Engine/Binaries/ThirdParty/VisualStudio/...
 	</Property>
 
 	<Property Name="CopyEditorFilterMac" Value="$(CopyEditorFilter);$(MacBuildFiles)">
@@ -458,7 +458,7 @@
 		Engine/Build/BatchFiles/MakeAndInstallSSHKey.bat
 		Engine/Build/Rsync/...
 		Engine/Extras/ThirdPartyNotUE/DeltaCopy/...
-		Engine/Extras/ThirdPartyNotUE/libimobiledevice/x64/...
+		Engine/Extras/ThirdPartyNotUE/libimobiledevice/x64/...	
 	</Property>
 
 	<Property Name="CopyIOSExceptions">
diff --git a/Engine/Build/Rsync/RsyncEngine.txt b/Engine/Build/Rsync/RsyncEngine.txt
--- a/Engine/Build/Rsync/RsyncEngine.txt
+++ b/Engine/Build/Rsync/RsyncEngine.txt
@@ -69,6 +69,7 @@ protect /Binaries/TVOS/
 
 # Build
 + /Build/
+- /Build/InstalledBuild.txt
 + /Build/*
 + /Build/BatchFiles/
 + /Build/BatchFiles/**
