Added EpicGames.Core source to Installed builds to build UBT for the IOS remote toolchain
Added Eigen, libJPG, mtlpp and libsamplerate source to Installed builds for IOS building

diff --git a/Engine/Build/InstalledEngineFilters.xml b/Engine/Build/InstalledEngineFilters.xml
--- a/Engine/Build/InstalledEngineFilters.xml
+++ b/Engine/Build/InstalledEngineFilters.xml
@@ -189,6 +189,7 @@
 		Source/Programs/AutomationTool/...
 		Source/Programs/AutomationToolLauncher/...
 		Source/Programs/DotNETCommon/...
+		Source/Programs/Shared/EpicGames.Core/...
 		Source/Programs/UnrealBuildTool/...
 		Source/Programs/UnrealHeaderTool/...
 
@@ -449,6 +450,10 @@
 		<!-- Build Files -->
 		Engine/Build/IOS/...
 		Engine/Source/ThirdParty/IOS/SoundSwitch/...
+		Engine/Source/ThirdParty/Eigen/...
+		Engine/Source/ThirdParty/libJPG/...
+		Engine/Source/ThirdParty/mtlpp/...
+		Engine/Source/ThirdParty/libSampleRate/...
 	</Property>
 	<Property Name="CopyIOSFilterWin64" Value="$(CopyIOSFilter)">
 		Engine/Binaries/DotNET/IOS/openssl.exe
