Exclude the Bridge plugin from installed builds.

diff --git a/Engine/Build/InstalledEngineFilters.xml b/Engine/Build/InstalledEngineFilters.xml
--- a/Engine/Build/InstalledEngineFilters.xml
+++ b/Engine/Build/InstalledEngineFilters.xml
@@ -220,6 +220,7 @@
 		Engine/Programs/UnrealGameSync/...
 
 		<!-- Plugins -->
+		Engine/Plugins/Bridge/...
 		Engine/Plugins/Runtime/TwitchLiveStreaming/...
 		Engine/Plugins/Runtime/PacketHandlers/CompressionComponents/Oodle/...
 		Engine/Platforms/*/Plugins/Runtime/PacketHandlers/CompressionComponents/Oodle/...
