Removing templates and feature packs that are deprecated in UE5.
Updating InstalledEngineFilters based on deprecations.

diff --git a/Engine/Build/InstalledEngineFilters.xml b/Engine/Build/InstalledEngineFilters.xml
--- a/Engine/Build/InstalledEngineFilters.xml
+++ b/Engine/Build/InstalledEngineFilters.xml
@@ -111,37 +111,24 @@
 
 		<!-- Starter content -->
 		Samples/StarterContent/Content/...
-		Samples/MobileStarterContent/Content/...
 
 		<!-- Templates -->
+
 		Templates/TemplateCategories.ini
-		Templates/FP_FirstPerson/...
-		Templates/FP_FirstPersonBP/...
 		Templates/Media/...
+
+		<!-- Game Templates -->
 		Templates/TP_Blank/...
 		Templates/TP_BlankBP/...
+		Templates/TP_HandheldARBP/...
 		Templates/TP_FirstPerson/...
 		Templates/TP_FirstPersonBP/...
-		Templates/TP_Flying/...
-		Templates/TP_FlyingBP/...
-		Templates/TP_HandheldARBP/...
-		Templates/TP_ProductConfigBP/...
-		Templates/TP_Rolling/...
-		Templates/TP_RollingBP/...
-		Templates/TP_SideScroller/...
-		Templates/TP_SideScrollerBP/...
 		Templates/TP_ThirdPerson/...
 		Templates/TP_ThirdPersonBP/...
 		Templates/TP_TopDown/...
 		Templates/TP_TopDownBP/...
-		Templates/TP_TwinStick/...
-		Templates/TP_TwinStickBP/...
-		Templates/TP_Vehicle/...
-		Templates/TP_VehicleBP/...
 		Templates/TP_Puzzle/...
 		Templates/TP_PuzzleBP/...
-		Templates/TP_2DSideScroller/...
-		Templates/TP_2DSideScrollerBP/...
 		Templates/TP_VehicleAdv/...
 		Templates/TP_VehicleAdvBP/...
 		Templates/TP_VirtualRealityBP/...
@@ -150,10 +137,9 @@
 		Templates/TP_AEC_BlankBP/...
 		Templates/TP_AEC_ArchvisBP/...
 		Templates/TP_PhotoStudioBP/...
-
+		Templates/TP_ProductConfigBP/...
 		Templates/TP_ME_BlankBP/...
 		Templates/TP_ME_VProdBP/...
-
 		Templates/TP_CollaborativeBP/...
 
 		<!-- Shared template resources -->
@@ -288,11 +274,8 @@
 		Engine/Extras/MayaVelocityGridExporter/...
 		Engine/Extras/UnrealVS/...
 		Engine/Extras/Redist/en-us/*
-			
-		<!-- Virtual Production Templates -->
-		Templates/TP_nDisplayBP/...
 
-		<!-- Enterprise Templates -->
+		<!-- Enterprise Templates Win64 Only -->
 		Templates/TP_HololensBP/...
 		Templates/TP_InCamVFXBP/...
 		Templates/TP_DMXBP/...
@@ -407,7 +390,6 @@
 	<!-- Mac -->
 	<Property Name="CopyMacFilter">
 		Engine/Binaries/Mac/BootstrapPackagedGame.app/...
-		Engine/Extras/ThirdPartyNotUE/libimobiledevice/Mac/...
 	</Property>
 	<Property Name="CopyMacExceptions">
 	</Property>
