Updating feature pack and DDC lists to exclude removed/deprecated templates and feature packs.

diff --git a/Engine/Build/InstalledEngineFilters.xml b/Engine/Build/InstalledEngineFilters.xml
--- a/Engine/Build/InstalledEngineFilters.xml
+++ b/Engine/Build/InstalledEngineFilters.xml
@@ -39,46 +39,26 @@
 
 	<!-- List of project to build Feature packs -->
 	<Property Name="ProjectsToFeaturePack">
-		FP_FirstPerson
-		FP_FirstPersonBP
-		TP_2DSideScroller
-		TP_2DSideScrollerBP
-		TP_Flying
-		TP_FlyingBP
 		TP_HandheldARBP
 		TP_Puzzle
 		TP_PuzzleBP
-		TP_Rolling
-		TP_RollingBP
-		TP_SideScroller
-		TP_SideScrollerBP
 		TP_ThirdPerson
 		TP_ThirdPersonBP
 		TP_TopDown
 		TP_TopDownBP
-		TP_TwinStick
-		TP_TwinStickBP
-		TP_Vehicle
 		TP_VehicleAdv
 		TP_VehicleAdvBP
-		TP_VehicleBP
 		TP_VirtualRealityBP
 	</Property>
 
 	<!-- List of projects to build DDC for -->
 	<Property Name="ProjectsToBuildDDC">
-		Templates/FP_FirstPersonBP/FP_FirstPersonBP.uproject
-		Templates/TP_FlyingBP/TP_FlyingBP.uproject
 		Templates/TP_HandheldARBP/TP_HandheldARBP.uproject
-		Templates/TP_RollingBP/TP_RollingBP.uproject
-		Templates/TP_SideScrollerBP/TP_SideScrollerBP.uproject
 		Templates/TP_ThirdPersonBP/TP_ThirdPersonBP.uproject
+		Templates/TP_FirstPersonBP/TP_FirstPersonBP.uproject
 		Templates/TP_TopDownBP/TP_TopDownBP.uproject
-		Templates/TP_TwinStickBP/TP_TwinStickBP.uproject
-		Templates/TP_VehicleBP/TP_VehicleBP.uproject
 		Templates/TP_ProductConfigBP/TP_ProductConfigBP.uproject
 		Templates/TP_PuzzleBP/TP_PuzzleBP.uproject
-		Templates/TP_2DSideScrollerBP/TP_2DSideScrollerBP.uproject
 		Templates/TP_VehicleAdvBP/TP_VehicleAdvBP.uproject
 		Templates/TP_AEC_BlankBP/TP_AEC_BlankBP.uproject
 		Templates/TP_AEC_ArchvisBP/TP_AEC_ArchvisBP.uproject
