Reverting changes to installed build filters probably made by accident.

diff --git a/Engine/Build/InstalledEngineFilters.xml b/Engine/Build/InstalledEngineFilters.xml
--- a/Engine/Build/InstalledEngineFilters.xml
+++ b/Engine/Build/InstalledEngineFilters.xml
@@ -155,7 +155,7 @@
 		Source/Programs/AutomationTool/...
 		Source/Programs/AutomationToolLauncher/...
 		Source/Programs/DotNETCommon/...
-		Source/Programs/Shared/EpicGames.Core/...
+		Source/Programs/Shared/...
 		Source/Programs/UnrealBuildTool/...
 		Source/Programs/UnrealHeaderTool/...
 
@@ -370,6 +370,7 @@
 	<!-- Mac -->
 	<Property Name="CopyMacFilter">
 		Engine/Binaries/Mac/BootstrapPackagedGame.app/...
+		Engine/Extras/ThirdPartyNotUE/libimobiledevice/Mac/...
 	</Property>
 	<Property Name="CopyMacExceptions">
 	</Property>
@@ -385,8 +386,6 @@
 		Engine/Source/ThirdParty/GoogleGameSDK/...
 		Engine/Source/ThirdParty/GoogleVR/...
 		Engine/Source/ThirdParty/GoogleInstantPreview/InstantPreview.apk
-		Engine/Source/ThirdParty/heapprofd/...
-		Engine/Source/ThirdParty/HWCPipe/...
 		Engine/Source/ThirdParty/Oculus/LibOVRAudio/LibOVRAudio/lib/*/*.so
 		Engine/Source/ThirdParty/Oculus/LibOVRAudio/*.tps
 		Engine/Source/ThirdParty/Oculus/LibOVRAudio/*.xml
