Fix "Build Tools CS Mac" on Mac

RunUAT.sh:
Setup Mono environment on Mac and Linux so that AutomationTool is able to launch xbuild for CsCompile tasks that have not been transitioned to NET Core.
(Linux change is untested)

CsCompileTask.cs:
Don't add the "/restore" argument to command line when building using xbuild as it does not appear to be supported.

InstalledEngineBuild.xml:
Remove CsCompile invocations for UAT projects as these are currently not able to be compiled by UAT: CsCompile does not know how use the bundled dotnet to build NET Core projects, and [our bundled] mono does not support NET Core projects.

This file is consumed by AutomationTool, which is compiled by RunUAT.sh. As such, it can generally be assumed that the UAT build is up to date when the built files are consumed (InstalledEngineBuild.xml has directives to tag the AutomationTool files to ensure they are included.)

diff --git a/Engine/Build/BatchFiles/RunUAT.sh b/Engine/Build/BatchFiles/RunUAT.sh
--- a/Engine/Build/BatchFiles/RunUAT.sh
+++ b/Engine/Build/BatchFiles/RunUAT.sh
@@ -67,11 +67,13 @@ UATDirectory=Binaries/DotNET/AutomationTool
 
 if [ "$(uname)" = "Darwin" ]; then
 	# Setup Environment
+	source "$SCRIPT_DIR/Mac/SetupEnvironment.sh" -mono "$SCRIPT_DIR/Mac"
 	source "$SCRIPT_DIR/Mac/SetupEnvironment.sh" $EnvironmentType "$SCRIPT_DIR/Mac"
 fi
 
 if [ "$(uname)" = "Linux" ]; then
 	# Setup Environment
+	source "$SCRIPT_DIR/Linux/SetupEnvironment.sh" -mono "$SCRIPT_DIR/Linux"
 	source "$SCRIPT_DIR/Linux/SetupEnvironment.sh" $EnvironmentType "$SCRIPT_DIR/Linux"
 fi
 
diff --git a/Engine/Build/InstalledEngineBuild.xml b/Engine/Build/InstalledEngineBuild.xml
--- a/Engine/Build/InstalledEngineBuild.xml
+++ b/Engine/Build/InstalledEngineBuild.xml
@@ -549,7 +549,6 @@
 
 			<!-- Copy Source and referenced libraries to a new location with Confidential folders removed-->
 			<Tag Files="$(AutomationToolPaths)" Filter="*.csproj" Except="$(ConfidentialExceptions)"  With="#UAT Projects Mac"/>
-			<CsCompile Project="#UAT Projects Mac" Configuration="Development" Platform="AnyCPU" TagReferences="#UATReferencesMac" EnumerateOnly="true"/>
 			<Tag Files="Engine/Binaries/DotNET/...;Engine/Binaries/ThirdParty/Newtonsoft/...;Engine/Binaries/ThirdParty/VisualStudio/...;Engine/Source/Programs/...;Engine/Platforms/*/Source/Programs/...;Engine/Source/Editor/SwarmInterface/...;#UATReferencesMac" Except="$(ConfidentialExceptions);AutomationTool.csproj.References" With="#RedistUATSourceMac"/>
 			<ForEach Name="Platform" Values="$(ExtensionPlatforms)">
 				<Expand Name="$(Platform)_TagCsToolsRedistUATSource"/>
@@ -562,9 +561,6 @@
 			</ForEach>
 			<Tag Files="$(RedistAutomationToolPaths)" Filter="*.csproj" With="#RedistUATProjectsMac"/>
 
-			<!-- Compile all the tools -->
-			<CsCompile Project="#RedistUATProjectsMac" Configuration="Development" Platform="AnyCPU" Tag="#Build Tools CS Binaries Mac;#Sign Binaries Mac" TagReferences="#Build Tools CS Binaries Mac"/>
-			
 			<!-- Tag AutomationTool and UnrealBuildTool folders recursively as NET Core dependencies are not currently handled by CsCompile -->
 			<Tag Files="$(CsToolsDir)/Engine/Binaries/DotNET/AutomationTool/..." With="#Build Tools CS Binaries Mac" />
 			<Tag Files="$(CsToolsDir)/Engine/Binaries/DotNET/UnrealBuildTool/..." With="#Build Tools CS Binaries Mac" />
diff --git a/Engine/Source/Programs/AutomationTool/BuildGraph/Tasks/CsCompileTask.cs b/Engine/Source/Programs/AutomationTool/BuildGraph/Tasks/CsCompileTask.cs
--- a/Engine/Source/Programs/AutomationTool/BuildGraph/Tasks/CsCompileTask.cs
+++ b/Engine/Source/Programs/AutomationTool/BuildGraph/Tasks/CsCompileTask.cs
@@ -159,7 +159,11 @@ namespace AutomationTool.Tasks
 				{
 					Arguments.Add(String.Format("/target:{0}", CommandUtils.MakePathSafeToUseWithCommandLine(Parameters.Target)));
 				}
-				Arguments.Add("/restore");
+				if(!CommandUtils.CmdEnv.MsBuildExe.Equals("xbuild"))
+				{
+					// not supported by xbuild
+					Arguments.Add("/restore");
+				}
 				Arguments.Add("/verbosity:minimal");
 				Arguments.Add("/nologo");
 				foreach(FileReference ProjectFile in ProjectFiles)
