Fix launcher iPhonePackager builds
Make Mac launcher use P4 version of iPhonePackager until we can either update Mono or move it to dotnet core.

diff --git a/Engine/Build/InstalledEngineBuild.xml b/Engine/Build/InstalledEngineBuild.xml
--- a/Engine/Build/InstalledEngineBuild.xml
+++ b/Engine/Build/InstalledEngineBuild.xml
@@ -480,7 +480,7 @@
 
 			<!-- Copy Source and referenced libraries to a new location with Confidential folders removed-->
 			<Tag Files="$(AutomationToolPaths)" Filter="*.csproj" Except="$(ConfidentialExceptions)"  With="#UAT Projects"/>
-			<Tag Files="Engine/Binaries/DotNET/...;Engine/Binaries/ThirdParty/Newtonsoft/...;Engine/Binaries/ThirdParty/VisualStudio/...;Engine/Source/Programs/...;Engine/Platforms/*/Source/Programs/...;Engine/Source/Editor/SwarmInterface/...;#UATReferences" Except="$(ConfidentialExceptions);AutomationTool.csproj.References" With="#RedistUATSource"/>
+			<Tag Files="Engine/Binaries/DotNET/...;Engine/Binaries/ThirdParty/Newtonsoft/...;Engine/Binaries/ThirdParty/IOS/...;Engine/Binaries/ThirdParty/VisualStudio/...;Engine/Source/Programs/...;Engine/Platforms/*/Source/Programs/...;Engine/Source/Editor/SwarmInterface/...;#UATReferences" Except="$(ConfidentialExceptions);AutomationTool.csproj.References" With="#RedistUATSource"/>
 			<ForEach Name="Platform" Values="$(ExtensionPlatforms)">
 				<Expand Name="$(Platform)_TagCsToolsRedistUATSource"/>
 			</ForEach>
@@ -509,13 +509,7 @@
 			<!-- Tag AutomationTool and UnrealBuildTool folders recursively as NET Core dependencies are not currently handled by CsCompile -->
 			<Tag Files="$(CsToolsDir)/Engine/Binaries/DotNET/AutomationTool/..." With="#Build Tools CS Binaries" />
 			<Tag Files="$(CsToolsDir)/Engine/Binaries/DotNET/UnrealBuildTool/..." With="#Build Tools CS Binaries" />
-			
-			<!-- Compile IPP. Mac/Linux bundled versions of Mono do not have Winforms, so we just use the existing one. -->
-			<Property Name="CompileIPP" Value="true"/>
-			<Property Name="CompileIPP" Value="false" If="'$(HostPlatform)' != 'Win64'"/>
-			<CsCompile Project="$(CsToolsDir)/Engine/Source/Programs/UnrealControls/UnrealControls.csproj" Configuration="Development" Platform="AnyCPU" Tag="#Build Tools CS Binaries;#Sign Binaries" TagReferences="#Build Tools CS Binaries" EnumerateOnly="!$(CompileIPP)"/>
-			<CsCompile Project="$(CsToolsDir)/Engine/Source/Programs/IOS/iPhonePackager/iPhonePackager.csproj" Configuration="Development" Platform="AnyCPU" Arguments="/verbosity:minimal /target:Rebuild" Tag="#Build Tools CS Binaries;#Sign Binaries" TagReferences="#Build Tools CS Binaries" EnumerateOnly="!$(CompileIPP)"/>
-			
+
 			<CsCompile Project="$(CsToolsDir)/Engine/Source/Programs/NetworkProfiler/NetworkProfiler/NetworkProfiler.csproj" Configuration="Development" Platform="AnyCPU" Tag="#Build Tools CS Binaries" TagReferences="#Build Tools CS Binaries;#Sign Binaries" If="'$(HostPlatform)' != 'Linux'"/>
 			<CsCompile Project="$(CsToolsDir)/Engine/Source/Programs/IOS/MobileDeviceInterface/MobileDeviceInterface.csproj" Configuration="Development" Platform="AnyCPU" Arguments="/verbosity:minimal /target:Rebuild" Tag="#Build Tools CS Binaries;#Sign Binaries" TagReferences="#Build Tools CS Binaries"/>
 
@@ -547,7 +541,7 @@
 
 			<!-- Copy Source and referenced libraries to a new location with Confidential folders removed-->
 			<Tag Files="$(AutomationToolPaths)" Filter="*.csproj" Except="$(ConfidentialExceptions)"  With="#UAT Projects Mac"/>
-			<Tag Files="Engine/Binaries/DotNET/...;Engine/Binaries/ThirdParty/Newtonsoft/...;Engine/Binaries/ThirdParty/VisualStudio/...;Engine/Source/Programs/...;Engine/Platforms/*/Source/Programs/...;Engine/Source/Editor/SwarmInterface/...;#UATReferencesMac" Except="$(ConfidentialExceptions);AutomationTool.csproj.References" With="#RedistUATSourceMac"/>
+			<Tag Files="Engine/Binaries/DotNET/...;Engine/Binaries/ThirdParty/Newtonsoft/...;Engine/Binaries/ThirdParty/IOS/...;Engine/Binaries/ThirdParty/VisualStudio/...;Engine/Source/Programs/...;Engine/Platforms/*/Source/Programs/...;Engine/Source/Editor/SwarmInterface/...;#UATReferencesMac" Except="$(ConfidentialExceptions);AutomationTool.csproj.References" With="#RedistUATSourceMac"/>
 			<ForEach Name="Platform" Values="$(ExtensionPlatforms)">
 				<Expand Name="$(Platform)_TagCsToolsRedistUATSource"/>
 			</ForEach>
@@ -566,6 +560,10 @@
 			<CsCompile Project="$(CsToolsDir)/Engine/Source/Programs/NetworkProfiler/NetworkProfiler/NetworkProfiler.csproj" Configuration="Development" Platform="AnyCPU" Tag="#Build Tools CS Binaries Mac" TagReferences="#Build Tools CS Binaries Mac;#Sign Binaries Mac"/>
 			<CsCompile Project="$(CsToolsDir)/Engine/Source/Programs/IOS/MobileDeviceInterface/MobileDeviceInterface.csproj" Configuration="Development" Platform="AnyCPU" Arguments="/verbosity:minimal /target:Rebuild" Tag="#Build Tools CS Binaries Mac;#Sign Binaries Mac" TagReferences="#Build Tools CS Binaries Mac"/>
 
+			<!-- IPP doesn't currently compile with the bundled Mono. Will ship the binary out of Perforce for now -->
+			<!-- CsCompile Project="$(CsToolsDir)/Engine/Source/Programs/UnrealControls/UnrealControls.csproj" Configuration="Development" Platform="AnyCPU" Tag="#Build Tools CS Binaries Mac;#Sign Binaries Mac" TagReferences="#Build Tools CS Binaries Mac"/-->
+			<!-- CsCompile Project="$(CsToolsDir)/Engine/Source/Programs/IOS/iPhonePackager/iPhonePackager.csproj" Configuration="Development" Platform="AnyCPU" Arguments="/verbosity:minimal /target:Rebuild" Tag="#Build Tools CS Binaries Mac;#Sign Binaries Mac" TagReferences="#Build Tools CS Binaries Mac"/-->
+
 			<!-- Sign the binaries -->
 			<Sign Files="#Sign Binaries Mac" If="'$(SignExecutables)' == true"/>
 		</Node>
@@ -704,6 +702,9 @@
 			<Tag Files="#ToCopy" Filter="*.target" With="#TargetReceipts"/>
 			<TagReceipt Files="#TargetReceipts" RuntimeDependencies="true" With="#ToCopy"/>
 			<Tag Files="Engine/Binaries/DotNET/Ionic.Zip.Reduced.dll" With="#ToCopy"/>
+			<Tag Files="Engine/Binaries/DotNET/IOS/IPhonePackager.exe" With="#ToCopy"/>
+			<Tag Files="Engine/Binaries/DotNET/IOS/IPhonePackager.exe.config" With="#ToCopy"/>
+			<Tag Files="Engine/Binaries/DotNET/IOS/IPhonePackager.pdb" With="#ToCopy"/>
 			<Tag Files="Engine/Binaries/DotNET/OneSky.dll" With="#ToCopy"/>
 			<Tag Files="Templates/TemplateResources/..." With="#ToCopy"/>
 			<Expand Name="TagDDCEngineFiles" ExceptPlugins="Plugins/.../Binaries/Win64/..."/>
