Fix launcher iPhonePackager builds
Make Mac launcher use P4 version of iPhonePackager until we can either update Mono or move it to dotnet core.
Moved copy from "DDC Group Mac" to "Build Tools CS Mac"

diff --git a/Engine/Build/InstalledEngineBuild.xml b/Engine/Build/InstalledEngineBuild.xml
--- a/Engine/Build/InstalledEngineBuild.xml
+++ b/Engine/Build/InstalledEngineBuild.xml
@@ -560,10 +560,19 @@
 			<CsCompile Project="$(CsToolsDir)/Engine/Source/Programs/NetworkProfiler/NetworkProfiler/NetworkProfiler.csproj" Configuration="Development" Platform="AnyCPU" Tag="#Build Tools CS Binaries Mac" TagReferences="#Build Tools CS Binaries Mac;#Sign Binaries Mac"/>
 			<CsCompile Project="$(CsToolsDir)/Engine/Source/Programs/IOS/MobileDeviceInterface/MobileDeviceInterface.csproj" Configuration="Development" Platform="AnyCPU" Arguments="/verbosity:minimal /target:Rebuild" Tag="#Build Tools CS Binaries Mac;#Sign Binaries Mac" TagReferences="#Build Tools CS Binaries Mac"/>
 
-			<!-- IPP doesn't currently compile with the bundled Mono. Will ship the binary out of Perforce for now -->
+			<!-- IPP doesn't currently compile with the bundled Mono. -->
 			<!-- CsCompile Project="$(CsToolsDir)/Engine/Source/Programs/UnrealControls/UnrealControls.csproj" Configuration="Development" Platform="AnyCPU" Tag="#Build Tools CS Binaries Mac;#Sign Binaries Mac" TagReferences="#Build Tools CS Binaries Mac"/-->
 			<!-- CsCompile Project="$(CsToolsDir)/Engine/Source/Programs/IOS/iPhonePackager/iPhonePackager.csproj" Configuration="Development" Platform="AnyCPU" Arguments="/verbosity:minimal /target:Rebuild" Tag="#Build Tools CS Binaries Mac;#Sign Binaries Mac" TagReferences="#Build Tools CS Binaries Mac"/-->
-
+		
+			<!-- Instead we will ship the binary out of Perforce for now. -->
+			<Tag Files="Engine/Binaries/DotNET/IOS/IPhonePackager.exe" With="#CopyIPPBinaries"/>
+			<Tag Files="Engine/Binaries/DotNET/IOS/IPhonePackager.exe.config" With="#CopyIPPBinaries"/>
+			<Tag Files="Engine/Binaries/DotNET/IOS/IPhonePackager.pdb" With="#CopyIPPBinaries"/>
+			<Copy Files="#CopyIPPBinaries" From="$(RootDir)" To="$(CsToolsDir)"/>
+			<Tag Files="$(CsToolsDir)/Engine/Binaries/DotNET/IOS/IPhonePackager.exe" With="#Build Tools CS Binaries Mac;#Sign Binaries Mac"/>
+			<Tag Files="$(CsToolsDir)/Engine/Binaries/DotNET/IOS/IPhonePackager.exe.config" With="#Build Tools CS Binaries Mac;#Sign Binaries Mac"/>
+			<Tag Files="$(CsToolsDir)/Engine/Binaries/DotNET/IOS/IPhonePackager.pdb" With="#Build Tools CS Binaries Mac;#Sign Binaries Mac"/>
+			
 			<!-- Sign the binaries -->
 			<Sign Files="#Sign Binaries Mac" If="'$(SignExecutables)' == true"/>
 		</Node>
@@ -702,9 +711,6 @@
 			<Tag Files="#ToCopy" Filter="*.target" With="#TargetReceipts"/>
 			<TagReceipt Files="#TargetReceipts" RuntimeDependencies="true" With="#ToCopy"/>
 			<Tag Files="Engine/Binaries/DotNET/Ionic.Zip.Reduced.dll" With="#ToCopy"/>
-			<Tag Files="Engine/Binaries/DotNET/IOS/IPhonePackager.exe" With="#ToCopy"/>
-			<Tag Files="Engine/Binaries/DotNET/IOS/IPhonePackager.exe.config" With="#ToCopy"/>
-			<Tag Files="Engine/Binaries/DotNET/IOS/IPhonePackager.pdb" With="#ToCopy"/>
 			<Tag Files="Engine/Binaries/DotNET/OneSky.dll" With="#ToCopy"/>
 			<Tag Files="Templates/TemplateResources/..." With="#ToCopy"/>
 			<Expand Name="TagDDCEngineFiles" ExceptPlugins="Plugins/.../Binaries/Win64/..."/>
