InstalledEngineFilters.xml:

Include RunUBT.sh for Mac and Linux

diff --git a/Engine/Build/InstalledEngineFilters.xml b/Engine/Build/InstalledEngineFilters.xml
--- a/Engine/Build/InstalledEngineFilters.xml
+++ b/Engine/Build/InstalledEngineFilters.xml
@@ -289,6 +289,7 @@
 		Engine/Build/BatchFiles/Mac/XcodeBuild.sh
 		Engine/Build/BatchFiles/RunUAT.command
 		Engine/Build/BatchFiles/RunUAT.sh
+		Engine/Build/BatchFiles/RunUBT.sh
 
 		<!-- Mono -->
 		Engine/Binaries/ThirdParty/Mono/Mac/...
@@ -327,6 +328,7 @@
 		Engine/Build/BatchFiles/Linux/RunMono.sh
 		Engine/Build/BatchFiles/Linux/RunXBuild.sh
 		Engine/Build/BatchFiles/RunUAT.sh
+		Engine/Build/BatchFiles/RunUBT.sh
 
 		<!-- Debug tools -->
 		Engine/Binaries/Linux/dump_syms
