Move IterationProfile out of Restricted

diff --git a/Engine/Build/BatchFiles/IterationProfile.bat b/Engine/Build/BatchFiles/IterationProfile.bat
new file mode 100755
--- /dev/null
+++ b/Engine/Build/BatchFiles/IterationProfile.bat
@@ -0,0 +1 @@
+@call %~dp0RunUAT.bat BuildGraph -Script=Engine/Build/IterationProfile.xml -Target="Run Iteration Profile Tests" -set:ProjectName=ShooterGame -set:TestWin64Editor=true %*
diff --git a/Engine/Build/IterationProfile.xml b/Engine/Build/IterationProfile.xml
new file mode 100644
--- /dev/null
+++ b/Engine/Build/IterationProfile.xml
@@ -0,0 +1,156 @@
+<?xml version='1.0' ?>
+<BuildGraph xmlns="http://www.epicgames.com/BuildGraph" xmlns:xsi="http://www.w3.org/2001/XMLSchema-instance" xsi:schemaLocation="http://www.epicgames.com/BuildGraph ../../../Build/Graph/Schema.xsd" >
+
+	<!-- Horde Options -->	
+	<EnvVar Name="UE_HORDE_JOBID"/>
+	<Option Name="Horde" Restrict="true|false" DefaultValue="false" Description="Whether the build is running under Horde"/>
+	<Option Name="PreflightChange" DefaultValue="" Description="Preflight CL number if preflight, empty otherwise"/>
+	<Option Name="SetPreflightFlag" DefaultValue="true" Description="Whether to set the IsPreflight flag to true for preflights. Use with caution: this will cause things like build versions to be set to their non-preflight form."/>
+	<Property Name="IsPreflight" Value="false"/>
+	<Property Name="IsPreflight" Value="true" If="'$(PreflightChange)' != '' And $(SetPreflightFlag)"/>
+	<Property Name="PreflightSuffix" Value=""/>
+	<Property Name="PreflightSuffix" Value="-PF-$(PreflightChange)-$(UE_HORDE_JOBID)" If="$(IsPreflight)"/>
+	<Property Name="BuildName" Value="$(EscapedBranch)-CL-$(Change)$(PreflightSuffix)"/>
+	
+	<!-- Project to run tests in -->
+	<Option Name="ProjectName" DefaultValue="false" Restrict=".+" Description="The Project to test"/>
+	<Error Message="No project name specified" If="'$(ProjectName)' == 'false'"/>
+	<Option Name="UgsProject" DefaultValue="" Description="The relative path to the project"/>
+	<Option Name="UgsBadge" DefaultValue="" Description="The badge name to use with labels"/>
+
+	<!-- Platform Options -->
+	<Option Name="TestWin64Editor" DefaultValue="false" Description="Run Win64 Editor Tests"/>
+	<Option Name="TestMacEditor" DefaultValue="false" Description="Run Mac Editor Tests"/>
+	<Error Message="No platforms specified" If="!$(TestWin64Editor) And !$(TestMacEditor)"/>
+
+	<!-- Build Options -->
+	<Property Name="ExtraPluginsCommaDelims" Value="EditorTests"/>
+	<Property Name="ExtraPluginsPlusDelims" Value=""/> <!-- populated programmatically -->
+	<ForEach Name="Plugin" Values="$(ExtraPluginsCommaDelims)">
+		<Property Name="ExtraPluginString" Value="$(Plugin)" If="'$(ExtraPluginsPlusDelims)' == ''"/>
+		<Property Name="ExtraPluginString" Value="$(ExtraPluginsPlusDelims)+$(Plugin)" If="'$(ExtraPluginsPlusDelims)' != ''"/>
+		<Property Name="ExtraPluginsPlusDelims" Value="$(ExtraPluginString)"/>
+	</ForEach>
+	<Property Name="EditorCompileArgs" Value="-BuildPlugin=$(ExtraPluginsPlusDelims)"/>
+
+	<!-- Automation uses its own network output directory -->
+	<Property Name="AutomationNodes" Value="" />
+	<Property Name="AutomationLocalOutputDir" Value="$(RootDir)/LocalBuilds/Reports"/>
+	<Property Name="AutomationNetworkOutputDir" Value="P:/Builds/Automation/Reports/$(EscapedBranch)/$(ProjectName)/$(BuildName)"/>
+	<Property Name="AutomationNetworkOutputDir" Value="/Volumes/Builds/Automation/Reports/$(EscapedBranch)/$(ProjectName)/$(BuildName)" If="'$(HostPlatform)' != 'Win64'"/>
+	<Property Name="AutomationReportBaseUrl" Value="http://automation.epicgames.net/reports/$(EscapedBranch)/$(ProjectName)/$(BuildName)"/>
+	<Property Name="AutomationOutputDir" Value="$(AutomationLocalOutputDir)"/>
+	<Property Name="AutomationOutputDir" Value="$(AutomationNetworkOutputDir)" If="$(IsBuildMachine)"/>
+
+	<!-- Arguments for tracing -->
+	<Option Name="ColdPass" DefaultValue="true" Description="Perform cold workspace loading profiles"/>
+	<Option Name="ColdArgs" DefaultValue="-ddc=Cold" Description="Arguments to unreal for the cold pass"/>
+	<Option Name="WarmPass" DefaultValue="true" Description="Perform warm (incremental) workspace loading profiles"/>
+	<Option Name="WarmArgs" DefaultValue="-ddc=Warm" Description="Arguments to unreal for the warm pass"/>
+	<Property Name="TraceDir" Value="$(RootDir)/Engine/Programs/UnrealInsights/Saved/TraceSessions"/>
+
+	<!-- Gauntlet arguments for all runs -->
+	<Property Name="GauntletArgs" Value="-build=editor -UseEditor -verbose -report -project=$(ProjectName) -tempdir=&quot;$(RootDir)/Tests&quot;"/>
+	<Property Name="GauntletArgs" Value="$(GauntletArgs) -branch=$(Branch) -ECBranch=&quot;$(Branch)&quot; -ECChangelist=&quot;$(Change)&quot;"/>
+	<Property Name="GauntletArgs" Value="$(GauntletArgs) -maxduration=14400"/>
+
+	<!-- Gauntlet arguments for profile runs -->
+	<Property Name="ProfileArgs" Value="-test=EditorTests -testfilter=Group:Iteration -tracefiletrunc"/>
+	<Property Name="ProfileArgs" Value="$(ProfileArgs) -EnablePlugins=$(ExtraPluginsCommaDelims)" If="'$(ExtraPluginsCommaDelims)' != ''"/>
+	<Property Name="ProfileArgs" Value="$(ProfileArgs) -TreatLogWarningsAsTestErrors=false -TreatLogErrorsAsTestErrors=false"/>
+	<Property Name="ProfileArgs" Value="$(ProfileArgs) -GameInstanceLostTimerSeconds=3600"/>
+	<Property Name="ProfileArgs" Value="$(ProfileArgs) -NoShaderDistrib"/>
+
+	<!-- Gauntlet arguments for summarize commandlet runs -->
+	<Property Name="SummarizeArgs" Value="-statsfile=&quot;$(RootDir)/Engine/Restricted/NotForLicensees/Build/IterationProfileStats.csv&quot;"/>
+
+	<!-- This will be set to the names of all nodes we want to execute based on specified platforms/options -->
+	<Property Name="GraphNodes" Value="" />
+	
+	<!-- The main test running macro logic that is expanded with platform and dependency node names -->
+	<Macro Name="RunTest" Arguments="OutputName;PlatformName;NodeName;NodeRequires;Phase;Args">
+		<Property Name="TraceFile" Value="$(TraceDir)/$(ProjectName)IterationProfile$(PlatformName)$(Phase).utrace"/>
+		<Property Name="TelemetryFile" Value="$(TraceDir)/$(ProjectName)IterationProfile$(PlatformName)$(Phase)Telemetry.csv"/>
+		<Property Name="ReportExportPath" Value="$(AutomationOutputDir)/$(OutputName)"/>
+		<Property Name="ReportParams" Value="-ReportExportPath=&quot;$(ReportExportPath)&quot;"/>
+		<Do If="$(IsBuildMachine)">
+			<Property Name="ReportExportURL" Value="$(AutomationReportBaseUrl)/$(OutputName)/index.html"/>
+			<Property Name="ReportParams" Value="$(ReportParams) -ReportURL=&quot;$(ReportExportURL)&quot; -WriteTestResultsForHorde -HordeTestDataKey=&quot;$(NodeName) $(Phase)$(PreflightSuffix)&quot;"/>
+		</Do>
+		<Property Name="LogOutputParams" Value="-logdir=&quot;$(AutomationOutputDir)/$(OutputName)&quot;"/>
+		<Property Name="AllGauntletArgs" Value="-platform=$(PlatformName) $(GauntletArgs) $(ReportParams) $(LogOutputParams) -artifactname=&quot;$(Phase)&quot;"/>
+		<Command Name="RunUnreal" Arguments="$(AllGauntletArgs) $(ProfileArgs) $(Args) -tracefile=&quot;$(TraceFile)&quot;"/>
+		<Commandlet Name="SummarizeTrace" Arguments="-NoEnginePlugins -inputfile=&quot;$(TraceFile)&quot; $(SummarizeArgs)"/>
+		<Do If="$(IsBuildMachine)">
+			<!-- Files copied to this folder are automatically uploaded to Horde as artifacts -->
+			<Copy From="$(TraceDir)/$(ProjectName)IterationProfile*.*" To="$(RootDir)/Engine/Programs/AutomationTool/Saved/Logs"/>
+			<Command Name="PublishUnrealAutomationTelemetry" Arguments="-CSVFile=&quot;$(TelemetryFile)&quot; -project=$(ProjectName) -changelist=$(Change) -branch=$(Branch)"/>
+		</Do>
+	</Macro>
+	
+	<!-- Windows prereq (code compiles) -->
+	<Agent Name="Prerequisites Agent Win64" Type="Win64">
+		<Node Name="Compile $(ProjectName) Prerequisites Win64Host">
+			<Compile Target="UnrealHeaderTool" Platform="Win64" Configuration="Development"/>
+			<Compile Target="ShaderCompileWorker" Platform="Win64" Configuration="Development"/>
+			<Compile Target="CrashReportClient" Platform="Win64" Configuration="Shipping"/>
+			<Compile Target="CrashReportClientEditor" Platform="Win64" Configuration="Shipping"/>
+		</Node>
+		<Node Name="Compile $(ProjectName) Editor Win64" Requires="Compile $(ProjectName) Prerequisites Win64Host">
+			<Compile Target="$(ProjectName)Editor" Platform="Win64" Configuration="Development" Arguments="$(EditorCompileArgs)"/>
+		</Node>
+	</Agent>
+	<Agent Name="Automated Test Agent Win64" Type="TestWin64">
+		<Property Name="PlatformName" Value="Win64" />
+		<Property Name="OutputName" Value="IterationProfile $(PlatformName)" />
+		<Property Name="NodeName" Value="Run $(ProjectName) $(OutputName)" />
+		<Property Name="NodeRequires" Value="Compile $(ProjectName) Prerequisites $(PlatformName)Host;Compile $(ProjectName) Editor $(PlatformName)" />
+		<Do If="$(TestWin64Editor)">
+			<Node Name="$(NodeName)" Requires="$(NodeRequires)">
+				<Do If="$(ColdPass)">
+					<Expand Name="RunTest" OutputName="$(OutputName)" PlatformName="$(PlatformName)" NodeName="$(NodeName)" NodeRequires="$(NodeRequires)" Phase="Cold" Args="$(ColdArgs)"/>
+				</Do>
+				<Do If="$(WarmPass)">
+					<Expand Name="RunTest" OutputName="$(OutputName)" PlatformName="$(PlatformName)" NodeName="$(NodeName)" NodeRequires="$(NodeRequires)" Phase="Warm" Args="$(WarmArgs)"/>
+				</Do>
+			</Node>
+			<Property Name="GraphNodes" Value="$(GraphNodes);$(NodeName)"/>
+		</Do>
+	</Agent>
+	
+	<!-- Mac prereq (code compiles) -->
+	<Agent Name="Prerequisites Agent Mac" Type="TestMac">
+		<Node Name="Compile $(ProjectName) Prerequisites Mac">
+			<Compile Target="UnrealHeaderTool" Platform="Mac" Configuration="Development"/>
+			<Compile Target="ShaderCompileWorker" Platform="Mac" Configuration="Development"/>
+			<Compile Target="CrashReportClient" Platform="Mac" Configuration="Shipping"/>
+			<Compile Target="CrashReportClientEditor" Platform="Mac" Configuration="Shipping"/>
+		</Node>
+		<Node Name="Compile $(ProjectName) Editor Mac" Requires="Compile $(ProjectName) Prerequisites Mac">
+			<Compile Target="$(ProjectName)Editor" Platform="Mac" Configuration="Development"  Arguments="$(EditorCompileArgs)"/>
+		</Node>
+	</Agent>
+	<Agent Name="Automated Test Agent Mac" Type="TestMac">
+		<Property Name="PlatformName" Value="Mac" />
+		<Property Name="OutputName" Value="IterationProfile $(PlatformName)" />
+		<Property Name="NodeName" Value="Run $(ProjectName) $(OutputName)" />
+		<Property Name="NodeRequires" Value="Compile $(ProjectName) Prerequisites $(PlatformName);Compile $(ProjectName) Editor $(PlatformName)" />
+		<Do If="$(TestMacEditor)">
+			<Node Name="$(NodeName)" Requires="$(NodeRequires)">
+				<Do If="$(ColdPass)">
+					<Expand Name="RunTest" OutputName="$(OutputName)" PlatformName="$(PlatformName)" NodeName="$(NodeName)" NodeRequires="$(NodeRequires)" Phase="Cold" Args="$(ColdArgs)"/>
+				</Do>
+				<Do If="$(WarmPass)">
+					<Expand Name="RunTest" OutputName="$(OutputName)" PlatformName="$(PlatformName)" NodeName="$(NodeName)" NodeRequires="$(NodeRequires)" Phase="Warm" Args="$(WarmArgs)"/>
+				</Do>
+			</Node>
+			<Property Name="GraphNodes" Value="$(GraphNodes);$(NodeName)"/>
+		</Do>
+	</Agent>
+
+	<!-- Graph nodes that really tie the room together -->
+	<Label Category="Tests" Name="Iteration Profile Win64" Requires="Prerequisites Agent Win64;Automated Test Agent Win64" UgsProject="$(UgsProject)" UgsBadge="$(UgsBadge)"/>
+	<Label Category="Tests" Name="Iteration Profile Mac" Requires="Prerequisites Agent Mac;Automated Test Agent Mac" UgsProject="$(UgsProject)" UgsBadge="$(UgsBadge)"/>
+	<Aggregate Name="Run Iteration Profile Tests" Requires="$(GraphNodes)"/>
+
+</BuildGraph>
