Support for running multiple projects within the same Iteration Profile build
* This will give us wider data collection to detect regressions in the cases where a single project has a blocking issue.

AutomationTool:
* Add support for appending semicolon delimited parameters into Options (to feed <ForEach/>)

Horde:
* Put project name into labels to isolate trends graph as builds are run on different projects

diff --git a/Engine/Build/BatchFiles/IterationProfile.bat b/Engine/Build/BatchFiles/IterationProfile.bat
--- a/Engine/Build/BatchFiles/IterationProfile.bat
+++ b/Engine/Build/BatchFiles/IterationProfile.bat
@@ -1 +1 @@
-@call %~dp0RunUAT.bat BuildGraph -Script=Engine/Build/IterationProfile.xml -Target="Run Iteration Profile Tests" -set:ProjectName=ShooterGame -set:TestWin64Editor=true %*
+@call %~dp0RunUAT.bat BuildGraph -Script=Engine/Build/IterationProfile.xml -Target="Run Iteration Profile Tests" -append:ProjectNames="EngineTest" -append:ProjectNames="ShooterGame" -set:TestWin64Editor=true %*
diff --git a/Engine/Build/IterationProfile.xml b/Engine/Build/IterationProfile.xml
--- a/Engine/Build/IterationProfile.xml
+++ b/Engine/Build/IterationProfile.xml
@@ -11,10 +11,10 @@
 	<Property Name="PreflightSuffix" Value=""/>
 	<Property Name="PreflightSuffix" Value="-PF-$(PreflightChange)-$(UE_HORDE_JOBID)" If="$(IsPreflight)"/>
 	<Property Name="BuildName" Value="$(EscapedBranch)-CL-$(Change)$(PreflightSuffix)"/>
-	
-	<!-- Project to run tests in -->
-	<Option Name="ProjectName" DefaultValue="false" Restrict=".+" Description="The Project to test"/>
-	<Error Message="No project name specified" If="'$(ProjectName)' == 'false'"/>
+
+	<!-- Projects setup-->
+	<Option Name="ProjectNames" DefaultValue="false" Restrict=".+" Description="The Projects to test"/>
+	<Error Message="No project names specified" If="'$(ProjectNames)' == 'false'"/>
 	<Option Name="UgsProject" DefaultValue="" Description="The relative path to the project"/>
 	<Option Name="UgsBadge" DefaultValue="" Description="The badge name to use with labels"/>
 
@@ -36,11 +36,6 @@
 	<!-- Automation uses its own network output directory -->
 	<Property Name="AutomationNodes" Value="" />
 	<Property Name="AutomationLocalOutputDir" Value="$(RootDir)/LocalBuilds/Reports"/>
-	<Property Name="AutomationNetworkOutputDir" Value="P:/Builds/Automation/Reports/$(EscapedBranch)/$(ProjectName)/$(BuildName)"/>
-	<Property Name="AutomationNetworkOutputDir" Value="/Volumes/Builds/Automation/Reports/$(EscapedBranch)/$(ProjectName)/$(BuildName)" If="'$(HostPlatform)' != 'Win64'"/>
-	<Property Name="AutomationReportBaseUrl" Value="http://automation.epicgames.net/reports/$(EscapedBranch)/$(ProjectName)/$(BuildName)"/>
-	<Property Name="AutomationOutputDir" Value="$(AutomationLocalOutputDir)"/>
-	<Property Name="AutomationOutputDir" Value="$(AutomationNetworkOutputDir)" If="$(IsBuildMachine)"/>
 
 	<!-- Arguments for tracing -->
 	<Option Name="ColdPass" DefaultValue="true" Description="Perform cold workspace loading profiles"/>
@@ -50,7 +45,7 @@
 	<Property Name="TraceDir" Value="$(RootDir)/Engine/Programs/UnrealInsights/Saved/TraceSessions"/>
 
 	<!-- Gauntlet arguments for all runs -->
-	<Property Name="GauntletArgs" Value="-build=editor -UseEditor -verbose -report -project=$(ProjectName) -tempdir=&quot;$(RootDir)/Tests&quot;"/>
+	<Property Name="GauntletArgs" Value="-build=editor -UseEditor -verbose -report -tempdir=&quot;$(RootDir)/Tests&quot;"/>
 	<Property Name="GauntletArgs" Value="$(GauntletArgs) -branch=$(Branch) -ECBranch=&quot;$(Branch)&quot; -ECChangelist=&quot;$(Change)&quot;"/>
 	<Property Name="GauntletArgs" Value="$(GauntletArgs) -maxduration=14400"/>
 
@@ -64,11 +59,8 @@
 	<!-- Gauntlet arguments for summarize commandlet runs -->
 	<Property Name="SummarizeArgs" Value="-statsfile=&quot;$(RootDir)/Engine/Restricted/NotForLicensees/Build/IterationProfileStats.csv&quot;"/>
 
-	<!-- This will be set to the names of all nodes we want to execute based on specified platforms/options -->
-	<Property Name="GraphNodes" Value="" />
-	
 	<!-- The main test running macro logic that is expanded with platform and dependency node names -->
-	<Macro Name="RunTest" Arguments="OutputName;PlatformName;NodeName;NodeRequires;Phase;Args">
+	<Macro Name="RunTest" Arguments="OutputName;PlatformName;ProjectName;NodeName;NodeRequires;Phase;Args">
 		<Property Name="TraceFile" Value="$(TraceDir)/$(ProjectName)IterationProfile$(PlatformName)$(Phase).utrace"/>
 		<Property Name="TelemetryFile" Value="$(TraceDir)/$(ProjectName)IterationProfile$(PlatformName)$(Phase)Telemetry.csv"/>
 		<Property Name="ReportExportPath" Value="$(AutomationOutputDir)/$(OutputName)"/>
@@ -78,7 +70,7 @@
 			<Property Name="ReportParams" Value="$(ReportParams) -ReportURL=&quot;$(ReportExportURL)&quot; -WriteTestResultsForHorde -HordeTestDataKey=&quot;$(NodeName) $(Phase)$(PreflightSuffix)&quot;"/>
 		</Do>
 		<Property Name="LogOutputParams" Value="-logdir=&quot;$(AutomationOutputDir)/$(OutputName)&quot;"/>
-		<Property Name="AllGauntletArgs" Value="-platform=$(PlatformName) $(GauntletArgs) $(ReportParams) $(LogOutputParams) -artifactname=&quot;$(Phase)&quot;"/>
+		<Property Name="AllGauntletArgs" Value="-platform=$(PlatformName) -project=$(ProjectName) $(GauntletArgs) $(ReportParams) $(LogOutputParams) -artifactname=&quot;$(Phase)&quot;"/>
 		<Command Name="RunUnreal" Arguments="$(AllGauntletArgs) $(ProfileArgs) $(Args) -tracefile=&quot;$(TraceFile)&quot;"/>
 		<Commandlet Name="SummarizeTrace" Arguments="-NoEnginePlugins -inputfile=&quot;$(TraceFile)&quot; $(SummarizeArgs)"/>
 		<Do If="$(IsBuildMachine)">
@@ -88,69 +80,93 @@
 		</Do>
 	</Macro>
 	
-	<!-- Windows prereq (code compiles) -->
+	<!-- This will be set to the names of all nodes we want to execute based on specified platforms/options -->
+	<Property Name="GraphNodes" Value="" />	
+
+	<!-- Windows -->
 	<Agent Name="Prerequisites Agent Win64" Type="Win64">
-		<Node Name="Compile $(ProjectName) Prerequisites Win64Host">
+		<Node Name="Compile Prerequisites Win64Host">
 			<Compile Target="UnrealHeaderTool" Platform="Win64" Configuration="Development"/>
 			<Compile Target="ShaderCompileWorker" Platform="Win64" Configuration="Development"/>
 			<Compile Target="CrashReportClient" Platform="Win64" Configuration="Shipping"/>
 			<Compile Target="CrashReportClientEditor" Platform="Win64" Configuration="Shipping"/>
 		</Node>
-		<Node Name="Compile $(ProjectName) Editor Win64" Requires="Compile $(ProjectName) Prerequisites Win64Host">
-			<Compile Target="$(ProjectName)Editor" Platform="Win64" Configuration="Development" Arguments="$(EditorCompileArgs)"/>
-		</Node>
 	</Agent>
-	<Agent Name="Automated Test Agent Win64" Type="TestWin64">
-		<Property Name="PlatformName" Value="Win64" />
-		<Property Name="OutputName" Value="IterationProfile $(PlatformName)" />
-		<Property Name="NodeName" Value="Run $(ProjectName) $(OutputName)" />
-		<Property Name="NodeRequires" Value="Compile $(ProjectName) Prerequisites $(PlatformName)Host;Compile $(ProjectName) Editor $(PlatformName)" />
-		<Do If="$(TestWin64Editor)">
-			<Node Name="$(NodeName)" Requires="$(NodeRequires)">
-				<Do If="$(ColdPass)">
-					<Expand Name="RunTest" OutputName="$(OutputName)" PlatformName="$(PlatformName)" NodeName="$(NodeName)" NodeRequires="$(NodeRequires)" Phase="Cold" Args="$(ColdArgs)"/>
-				</Do>
-				<Do If="$(WarmPass)">
-					<Expand Name="RunTest" OutputName="$(OutputName)" PlatformName="$(PlatformName)" NodeName="$(NodeName)" NodeRequires="$(NodeRequires)" Phase="Warm" Args="$(WarmArgs)"/>
-				</Do>
-			</Node>
-			<Property Name="GraphNodes" Value="$(GraphNodes);$(NodeName)"/>
-		</Do>
-	</Agent>
-	
-	<!-- Mac prereq (code compiles) -->
+
+	<!-- macOS -->
 	<Agent Name="Prerequisites Agent Mac" Type="TestMac">
-		<Node Name="Compile $(ProjectName) Prerequisites Mac">
+		<Node Name="Compile Prerequisites Mac">
 			<Compile Target="UnrealHeaderTool" Platform="Mac" Configuration="Development"/>
 			<Compile Target="ShaderCompileWorker" Platform="Mac" Configuration="Development"/>
 			<Compile Target="CrashReportClient" Platform="Mac" Configuration="Shipping"/>
 			<Compile Target="CrashReportClientEditor" Platform="Mac" Configuration="Shipping"/>
 		</Node>
-		<Node Name="Compile $(ProjectName) Editor Mac" Requires="Compile $(ProjectName) Prerequisites Mac">
-			<Compile Target="$(ProjectName)Editor" Platform="Mac" Configuration="Development"  Arguments="$(EditorCompileArgs)"/>
-		</Node>
 	</Agent>
-	<Agent Name="Automated Test Agent Mac" Type="TestMac">
-		<Property Name="PlatformName" Value="Mac" />
-		<Property Name="OutputName" Value="IterationProfile $(PlatformName)" />
-		<Property Name="NodeName" Value="Run $(ProjectName) $(OutputName)" />
-		<Property Name="NodeRequires" Value="Compile $(ProjectName) Prerequisites $(PlatformName);Compile $(ProjectName) Editor $(PlatformName)" />
-		<Do If="$(TestMacEditor)">
-			<Node Name="$(NodeName)" Requires="$(NodeRequires)">
-				<Do If="$(ColdPass)">
-					<Expand Name="RunTest" OutputName="$(OutputName)" PlatformName="$(PlatformName)" NodeName="$(NodeName)" NodeRequires="$(NodeRequires)" Phase="Cold" Args="$(ColdArgs)"/>
-				</Do>
-				<Do If="$(WarmPass)">
-					<Expand Name="RunTest" OutputName="$(OutputName)" PlatformName="$(PlatformName)" NodeName="$(NodeName)" NodeRequires="$(NodeRequires)" Phase="Warm" Args="$(WarmArgs)"/>
-				</Do>
+
+	<!-- ForEach Project to run, generate the graph nodes -->
+	<ForEach Name="ProjectName" Values="$(ProjectNames)">
+
+		<!-- Automation uses its own network output directory -->
+		<Property Name="AutomationNetworkOutputDir" Value="P:/Builds/Automation/Reports/$(EscapedBranch)/$(ProjectName)/$(BuildName)"/>
+		<Property Name="AutomationNetworkOutputDir" Value="/Volumes/Builds/Automation/Reports/$(EscapedBranch)/$(ProjectName)/$(BuildName)" If="'$(HostPlatform)' != 'Win64'"/>
+		<Property Name="AutomationReportBaseUrl" Value="http://automation.epicgames.net/reports/$(EscapedBranch)/$(ProjectName)/$(BuildName)"/>
+		<Property Name="AutomationOutputDir" Value="$(AutomationLocalOutputDir)"/>
+		<Property Name="AutomationOutputDir" Value="$(AutomationNetworkOutputDir)" If="$(IsBuildMachine)"/>
+
+		<!-- Windows -->
+		<Agent Name="Prerequisites Agent Win64" Type="Win64">
+			<Node Name="Compile $(ProjectName) Editor Win64" Requires="Compile Prerequisites Win64Host">
+				<Compile Target="$(ProjectName)Editor" Platform="Win64" Configuration="Development" Arguments="$(EditorCompileArgs)"/>
 			</Node>
-			<Property Name="GraphNodes" Value="$(GraphNodes);$(NodeName)"/>
-		</Do>
-	</Agent>
+		</Agent>
+		<Agent Name="Automated Test Agent Win64" Type="TestWin64">
+			<Property Name="PlatformName" Value="Win64" />
+			<Property Name="OutputName" Value="IterationProfile $(PlatformName)" />
+			<Property Name="NodeName" Value="Run $(ProjectName) $(OutputName)" />
+			<Property Name="NodeRequires" Value="Compile Prerequisites $(PlatformName)Host;Compile $(ProjectName) Editor $(PlatformName)" />
+			<Do If="$(TestWin64Editor)">
+				<Node Name="$(NodeName)" Requires="$(NodeRequires)">
+					<Do If="$(ColdPass)">
+						<Expand Name="RunTest" OutputName="$(OutputName)" PlatformName="$(PlatformName)" ProjectName="$(ProjectName)" NodeName="$(NodeName)" NodeRequires="$(NodeRequires)" Phase="Cold" Args="$(ColdArgs)"/>
+					</Do>
+					<Do If="$(WarmPass)">
+						<Expand Name="RunTest" OutputName="$(OutputName)" PlatformName="$(PlatformName)" ProjectName="$(ProjectName)" NodeName="$(NodeName)" NodeRequires="$(NodeRequires)" Phase="Warm" Args="$(WarmArgs)"/>
+					</Do>
+				</Node>
+				<Property Name="GraphNodes" Value="$(GraphNodes);$(NodeName)"/>
+			</Do>
+		</Agent>
+		
+		<!-- macOS -->
+		<Agent Name="Prerequisites Agent Mac" Type="TestMac">
+			<Node Name="Compile $(ProjectName) Editor Mac" Requires="Compile Prerequisites Mac">
+				<Compile Target="$(ProjectName)Editor" Platform="Mac" Configuration="Development"  Arguments="$(EditorCompileArgs)"/>
+			</Node>
+		</Agent>
+		<Agent Name="Automated Test Agent Mac" Type="TestMac">
+			<Property Name="PlatformName" Value="Mac" />
+			<Property Name="OutputName" Value="IterationProfile $(PlatformName)" />
+			<Property Name="NodeName" Value="Run $(ProjectName) $(OutputName)" />
+			<Property Name="NodeRequires" Value="Compile Prerequisites $(PlatformName);Compile $(ProjectName) Editor $(PlatformName)" />
+			<Do If="$(TestMacEditor)">
+				<Node Name="$(NodeName)" Requires="$(NodeRequires)">
+					<Do If="$(ColdPass)">
+						<Expand Name="RunTest" OutputName="$(OutputName)" PlatformName="$(PlatformName)" ProjectName="$(ProjectName)" NodeName="$(NodeName)" NodeRequires="$(NodeRequires)" Phase="Cold" Args="$(ColdArgs)"/>
+					</Do>
+					<Do If="$(WarmPass)">
+						<Expand Name="RunTest" OutputName="$(OutputName)" PlatformName="$(PlatformName)" ProjectName="$(ProjectName)" NodeName="$(NodeName)" NodeRequires="$(NodeRequires)" Phase="Warm" Args="$(WarmArgs)"/>
+					</Do>
+				</Node>
+				<Property Name="GraphNodes" Value="$(GraphNodes);$(NodeName)"/>
+			</Do>
+		</Agent>
+
+		<!-- Graph nodes that really tie the room together -->
+		<Label Category="Tests" Name="$(ProjectName) Iteration Profile Win64" Requires="Prerequisites Agent Win64;Automated Test Agent Win64" UgsProject="$(UgsProject)" UgsBadge="$(UgsBadge)"/>
+		<Label Category="Tests" Name="$(ProjectName) Iteration Profile Mac" Requires="Prerequisites Agent Mac;Automated Test Agent Mac" UgsProject="$(UgsProject)" UgsBadge="$(UgsBadge)"/>
+		
+	</ForEach>
 
-	<!-- Graph nodes that really tie the room together -->
-	<Label Category="Tests" Name="Iteration Profile Win64" Requires="Prerequisites Agent Win64;Automated Test Agent Win64" UgsProject="$(UgsProject)" UgsBadge="$(UgsBadge)"/>
-	<Label Category="Tests" Name="Iteration Profile Mac" Requires="Prerequisites Agent Mac;Automated Test Agent Mac" UgsProject="$(UgsProject)" UgsBadge="$(UgsBadge)"/>
 	<Aggregate Name="Run Iteration Profile Tests" Requires="$(GraphNodes)"/>
 
 </BuildGraph>
diff --git a/Engine/Source/Programs/AutomationTool/BuildGraph/BuildGraph.cs b/Engine/Source/Programs/AutomationTool/BuildGraph/BuildGraph.cs
--- a/Engine/Source/Programs/AutomationTool/BuildGraph/BuildGraph.cs
+++ b/Engine/Source/Programs/AutomationTool/BuildGraph/BuildGraph.cs
@@ -172,17 +172,40 @@ namespace AutomationTool
 			Dictionary<string, string> Arguments = new Dictionary<string, string>(StringComparer.InvariantCultureIgnoreCase);
 			foreach (string Param in Params)
 			{
-				const string Prefix = "set:";
-				if(Param.StartsWith(Prefix, StringComparison.InvariantCultureIgnoreCase))
+				const string SetPrefix = "set:";
+				if(Param.StartsWith(SetPrefix, StringComparison.InvariantCultureIgnoreCase))
 				{
 					int EqualsIdx = Param.IndexOf('=');
 					if(EqualsIdx >= 0)
 					{
-						Arguments[Param.Substring(Prefix.Length, EqualsIdx - Prefix.Length)] = Param.Substring(EqualsIdx + 1);
+						Arguments[Param.Substring(SetPrefix.Length, EqualsIdx - SetPrefix.Length)] = Param.Substring(EqualsIdx + 1);
 					}
 					else
 					{
-						LogWarning("Missing value for '{0}'", Param.Substring(Prefix.Length));
+						LogWarning("Missing value for '{0}'", Param.Substring(SetPrefix.Length));
+					}
+				}
+
+				const string AppendPrefix = "append:";
+				if(Param.StartsWith(AppendPrefix, StringComparison.InvariantCultureIgnoreCase))
+				{
+					int EqualsIdx = Param.IndexOf('=');
+					if(EqualsIdx >= 0)
+					{
+						string Property = Param.Substring(AppendPrefix.Length, EqualsIdx - AppendPrefix.Length);
+						string Value = Param.Substring(EqualsIdx + 1);
+						if (Arguments.ContainsKey(Property))
+						{
+							Arguments[Property] = Arguments[Property] + ";" + Value;
+						}
+						else
+						{
+							Arguments[Property] = Value;
+						}
+					}
+					else
+					{
+						LogWarning("Missing value for '{0}'", Param.Substring(AppendPrefix.Length));
 					}
 				}
 			}
