Parallelize Run phase of Iteration Profile

diff --git a/Engine/Build/IterationProfile.xml b/Engine/Build/IterationProfile.xml
--- a/Engine/Build/IterationProfile.xml
+++ b/Engine/Build/IterationProfile.xml
@@ -119,7 +119,7 @@
 				<Compile Target="$(ProjectName)Editor" Platform="Win64" Configuration="Development" Arguments="$(EditorCompileArgs)"/>
 			</Node>
 		</Agent>
-		<Agent Name="Automated Test Agent Win64" Type="TestWin64">
+		<Agent Name="Automated Test Agent $(ProjectName) Win64" Type="TestWin64">
 			<Property Name="PlatformName" Value="Win64" />
 			<Property Name="OutputName" Value="IterationProfile $(PlatformName)" />
 			<Property Name="NodeName" Value="Run $(ProjectName) $(OutputName)" />
@@ -143,7 +143,7 @@
 				<Compile Target="$(ProjectName)Editor" Platform="Mac" Configuration="Development"  Arguments="$(EditorCompileArgs)"/>
 			</Node>
 		</Agent>
-		<Agent Name="Automated Test Agent Mac" Type="TestMac">
+		<Agent Name="Automated Test Agent $(ProjectName) Mac" Type="TestMac">
 			<Property Name="PlatformName" Value="Mac" />
 			<Property Name="OutputName" Value="IterationProfile $(PlatformName)" />
 			<Property Name="NodeName" Value="Run $(ProjectName) $(OutputName)" />
@@ -162,8 +162,8 @@
 		</Agent>
 
 		<!-- Graph nodes that really tie the room together -->
-		<Label Category="Tests" Name="$(ProjectName) Iteration Profile Win64" Requires="Prerequisites Agent Win64;Automated Test Agent Win64" UgsProject="$(UgsProject)" UgsBadge="$(UgsBadge)"/>
-		<Label Category="Tests" Name="$(ProjectName) Iteration Profile Mac" Requires="Prerequisites Agent Mac;Automated Test Agent Mac" UgsProject="$(UgsProject)" UgsBadge="$(UgsBadge)"/>
+		<Label Category="Tests" Name="$(ProjectName) Iteration Profile Win64" Requires="Prerequisites Agent Win64;Automated Test Agent $(ProjectName) Win64" UgsProject="$(UgsProject)" UgsBadge="$(UgsBadge)"/>
+		<Label Category="Tests" Name="$(ProjectName) Iteration Profile Mac" Requires="Prerequisites Agent Mac;Automated Test Agent $(ProjectName) Mac" UgsProject="$(UgsProject)" UgsBadge="$(UgsBadge)"/>
 		
 	</ForEach>
 
