Fix relative pathing issue

diff --git a/Engine/Build/IterationProfile.xml b/Engine/Build/IterationProfile.xml
--- a/Engine/Build/IterationProfile.xml
+++ b/Engine/Build/IterationProfile.xml
@@ -57,7 +57,7 @@
 	<Property Name="ProfileArgs" Value="$(ProfileArgs) -NoShaderDistrib"/>
 
 	<!-- Gauntlet arguments for summarize commandlet runs -->
-	<Property Name="SummarizeArgs" Value="-statsfile=&quot;$(RootDir)/Engine/Restricted/NotForLicensees/Build/IterationProfileStats.csv&quot;"/>
+	<Property Name="SummarizeArgs" Value="-statsfile=&quot;$(RootDir)/Engine/Build/IterationProfileStats.csv&quot;"/>
 
 	<!-- The main test running macro logic that is expanded with platform and dependency node names -->
 	<Macro Name="RunTest" Arguments="OutputName;PlatformName;ProjectName;NodeName;NodeRequires;Phase;Args">
