Rename Win64 and MacOS support arguments.

Preliminary change for bless support. These arguments are general platform toggles now.

diff --git a/Engine/Build/BatchFiles/IterationProfile.bat b/Engine/Build/BatchFiles/IterationProfile.bat
--- a/Engine/Build/BatchFiles/IterationProfile.bat
+++ b/Engine/Build/BatchFiles/IterationProfile.bat
@@ -1 +1 @@
-@call %~dp0RunUAT.bat BuildGraph -Script=Engine/Build/IterationProfile.xml -Target="Run Iteration Profile Tests" -append:ProjectNames="EngineTest" -append:ProjectNames="ShooterGame" -set:TestWin64Editor=true %*
+@call %~dp0RunUAT.bat BuildGraph -Script=Engine/Build/IterationProfile.xml -Target="Run Iteration Profile Tests" -append:ProjectNames="EngineTest" -append:ProjectNames="ShooterGame" -set:Win64=true %*
diff --git a/Engine/Build/IterationProfile.xml b/Engine/Build/IterationProfile.xml
--- a/Engine/Build/IterationProfile.xml
+++ b/Engine/Build/IterationProfile.xml
@@ -19,9 +19,9 @@
 	<Option Name="UgsBadge" DefaultValue="" Description="The badge name to use with labels"/>
 
 	<!-- Platform Options -->
-	<Option Name="TestWin64Editor" DefaultValue="false" Description="Run Win64 Editor Tests"/>
-	<Option Name="TestMacEditor" DefaultValue="false" Description="Run Mac Editor Tests"/>
-	<Error Message="No platforms specified" If="!$(TestWin64Editor) And !$(TestMacEditor)"/>
+	<Option Name="Win64" DefaultValue="false" Description="Enable Win64 Platform"/>
+	<Option Name="MacOS" DefaultValue="false" Description="Enable MacOS Platform"/>
+	<Error Message="No platforms specified" If="!$(Win64) And !$(MacOS)"/>
 
 	<!-- Build Options -->
 	<Property Name="ExtraPluginsCommaDelims" Value="EditorTests"/>
@@ -125,7 +125,7 @@
 			<Property Name="OutputName" Value="IterationProfile $(PlatformName)" />
 			<Property Name="NodeName" Value="Run $(ProjectName) $(OutputName)" />
 			<Property Name="NodeRequires" Value="Compile Prerequisites $(PlatformName)Host;Compile $(ProjectName) Editor $(PlatformName)" />
-			<Do If="$(TestWin64Editor)">
+			<Do If="$(Win64)">
 				<Node Name="$(NodeName)" Requires="$(NodeRequires)">
 					<Do If="$(ColdPass)">
 						<Expand Name="RunTest" OutputName="$(OutputName)" PlatformName="$(PlatformName)" ProjectName="$(ProjectName)" NodeName="$(NodeName)" NodeRequires="$(NodeRequires)" Phase="Cold" Args="$(ColdArgs)"/>
@@ -149,7 +149,7 @@
 			<Property Name="OutputName" Value="IterationProfile $(PlatformName)" />
 			<Property Name="NodeName" Value="Run $(ProjectName) $(OutputName)" />
 			<Property Name="NodeRequires" Value="Compile Prerequisites $(PlatformName);Compile $(ProjectName) Editor $(PlatformName)" />
-			<Do If="$(TestMacEditor)">
+			<Do If="$(MacOS)">
 				<Node Name="$(NodeName)" Requires="$(NodeRequires)">
 					<Do If="$(ColdPass)">
 						<Expand Name="RunTest" OutputName="$(OutputName)" PlatformName="$(PlatformName)" ProjectName="$(ProjectName)" NodeName="$(NodeName)" NodeRequires="$(NodeRequires)" Phase="Cold" Args="$(ColdArgs)"/>
