Include the DotNet binaries for Windows and Linux.

diff --git a/Engine/Build/InstalledEngineFilters.xml b/Engine/Build/InstalledEngineFilters.xml
--- a/Engine/Build/InstalledEngineFilters.xml
+++ b/Engine/Build/InstalledEngineFilters.xml
@@ -264,6 +264,10 @@
 		Templates/TP_InCamVFXBP/...
 		Templates/TP_DMXBP/...
 
+		<!-- DotNet -->
+		Engine/Binaries/ThirdParty/DotNet/Windows/...
+		Engine/Binaries/ThirdParty/VisualStudio/...
+
 		<!-- CsvTools binaries which are manually updated in P4 -->
 		Engine/Binaries/DotNET/CsvTools/...
 	</Property>
@@ -339,6 +343,10 @@
 		<!-- Mono -->
 		Engine/Binaries/ThirdParty/Mono/...
 
+		<!-- DotNet -->
+		Engine/Binaries/ThirdParty/DotNet/Linux/...
+		Engine/Binaries/ThirdParty/VisualStudio/...
+
 		<!-- Extra tools -->
 		Engine/Extras/GDBPrinters/UEPrinters.py
 
