Consolidating IterationTest to a single agent and workspace type until we determine if this fixes things.

#rnx

diff --git a/Engine/Build/IterationProfile.xml b/Engine/Build/IterationProfile.xml
--- a/Engine/Build/IterationProfile.xml
+++ b/Engine/Build/IterationProfile.xml
@@ -158,7 +158,7 @@
 				<Compile Target="$(ProjectName)Editor" Platform="Win64" Configuration="Development" Arguments="$(EditorCompileArgs)"/>
 			</Node>
 		</Agent>
-		<Agent Name="Automated Test Agent $(ProjectName) Win64" Type="TestWin64">
+		<Agent Name="Automated Test Agent $(ProjectName) Win64" Type="IterationTest;TestWin64">
 			<Property Name="PlatformName" Value="Win64"/>
 			<Property Name="OutputName" Value="$(ProjectName)$(PlatformName)IterationProfile"/>
 			<Property Name="RunNodeName" Value="Run $(OutputName)"/>
