Add TP_VirtualRealityBP into list of projects to build DDC for.

diff --git a/Engine/Build/InstalledEngineFilters.xml b/Engine/Build/InstalledEngineFilters.xml
--- a/Engine/Build/InstalledEngineFilters.xml
+++ b/Engine/Build/InstalledEngineFilters.xml
@@ -53,19 +53,20 @@
 
 	<!-- List of projects to build DDC for -->
 	<Property Name="ProjectsToBuildDDC">
-		Templates/TP_HandheldARBP/TP_HandheldARBP.uproject
-		Templates/TP_ThirdPersonBP/TP_ThirdPersonBP.uproject
-		Templates/TP_FirstPersonBP/TP_FirstPersonBP.uproject
-		Templates/TP_TopDownBP/TP_TopDownBP.uproject
+		Templates/TP_AEC_ArchvisBP/TP_AEC_ArchvisBP.uproject
+		Templates/TP_AEC_BlankBP/TP_AEC_BlankBP.uproject
+		Templates/TP_AEC_CollabBP/TP_AEC_CollabBP.uproject
 		Templates/TP_AEC_ProdConfigBP/TP_AEC_ProdConfigBP.uproject
+		Templates/TP_FirstPersonBP/TP_FirstPersonBP.uproject
+		Templates/TP_HandheldARBP/TP_HandheldARBP.uproject
+		Templates/TP_MFG_CollabBP/TP_MFG_CollabBP.uproject
 		Templates/TP_MFG_ProdConfigBP/TP_MFG_ProdConfigBP.uproject
+		Templates/TP_PhotoStudioBP/TP_PhotoStudioBP.uproject
 		Templates/TP_PuzzleBP/TP_PuzzleBP.uproject
+		Templates/TP_ThirdPersonBP/TP_ThirdPersonBP.uproject
+		Templates/TP_TopDownBP/TP_TopDownBP.uproject
 		Templates/TP_VehicleAdvBP/TP_VehicleAdvBP.uproject
-		Templates/TP_AEC_BlankBP/TP_AEC_BlankBP.uproject
-		Templates/TP_AEC_ArchvisBP/TP_AEC_ArchvisBP.uproject
-		Templates/TP_PhotoStudioBP/TP_PhotoStudioBP.uproject
-		Templates/TP_AEC_CollabBP/TP_AEC_CollabBP.uproject
-		Templates/TP_MFG_CollabBP/TP_MFG_CollabBP.uproject
+		Templates/TP_VirtualRealityBP/TP_VirtualRealityBP.uproject
 		Samples/StarterContent/StarterContent.uproject
 	</Property>
 	<!-- Win64 specific DDC -->
@@ -119,12 +120,10 @@
 		Templates/TP_AEC_BlankBP/...
 		Templates/TP_AEC_ArchvisBP/...
 		Templates/TP_PhotoStudioBP/...
-		Templates/TP_AEC_ProdConfigBP/...
-		Templates/TP_MFG_ProdConfigBP/...
+		Templates/TP_ProductConfigBP/...
 		Templates/TP_ME_BlankBP/...
 		Templates/TP_ME_VProdBP/...
-		Templates/TP_AEC_CollabBP/...
-		Templates/TP_MFG_CollabBP/...
+		Templates/TP_CollaborativeBP/...
 
 		<!-- Shared template resources -->
 		Templates/TemplateResources/...
@@ -400,6 +399,7 @@
 		Engine/Source/ThirdParty/GoogleGameSDK/...
 		Engine/Source/ThirdParty/GoogleVR/...
 		Engine/Source/ThirdParty/GoogleInstantPreview/InstantPreview.apk
+		Engine/Source/ThirdParty/HWCPipe/...
 		Engine/Source/ThirdParty/Oculus/LibOVRAudio/LibOVRAudio/lib/*/*.so
 		Engine/Source/ThirdParty/Oculus/LibOVRAudio/*.tps
 		Engine/Source/ThirdParty/Oculus/LibOVRAudio/*.xml
