Fix installed builds for Linux missing SetupDotnet.sh

This is from a UDN

diff --git a/Engine/Build/InstalledEngineFilters.xml b/Engine/Build/InstalledEngineFilters.xml
--- a/Engine/Build/InstalledEngineFilters.xml
+++ b/Engine/Build/InstalledEngineFilters.xml
@@ -329,6 +329,7 @@
 		Engine/Build/BatchFiles/Linux/Setup.sh
 		Engine/Build/BatchFiles/Linux/SetupEnvironment.sh
 		Engine/Build/BatchFiles/Linux/SetupMono.sh
+		Engine/Build/BatchFiles/Linux/SetupDotnet.sh
 		Engine/Build/BatchFiles/Linux/FixMonoFiles.sh
 		Engine/Build/BatchFiles/Linux/RunMono.sh
 		Engine/Build/BatchFiles/Linux/RunXBuild.sh
