Set the compatible changelist to the same CL as EA1.

diff --git a/Engine/Build/Build.version b/Engine/Build/Build.version
--- a/Engine/Build/Build.version
+++ b/Engine/Build/Build.version
@@ -3,7 +3,7 @@
 	"MinorVersion": 0,
 	"PatchVersion": 0,
 	"Changelist": 0,
-	"CompatibleChangelist": 0,
+	"CompatibleChangelist": 16433597,
 	"IsLicenseeVersion": 0,
 	"IsPromotedBuild": 0,
 	"BranchName": "UE5"
