Fix Label category

#rnx

diff --git a/Engine/Build/IterationProfile.xml b/Engine/Build/IterationProfile.xml
--- a/Engine/Build/IterationProfile.xml
+++ b/Engine/Build/IterationProfile.xml
@@ -229,8 +229,8 @@
 		</Agent>
 
 		<!-- Graph nodes that really tie the room together -->
-		<Label Category="Tests" Name="$(ProjectName) Iteration Profile Win64" Requires="Prerequisites Agent Win64;Automated Test Agent $(ProjectName) Win64" UgsProject="$(UgsProject)" UgsBadge="$(UgsBadge)"/>
-		<Label Category="Tests" Name="$(ProjectName) Iteration Profile Mac" Requires="Prerequisites Agent Mac;Automated Test Agent $(ProjectName) Mac" UgsProject="$(UgsProject)" UgsBadge="$(UgsBadge)"/>
+		<Label Category="Test" Name="$(ProjectName) Iteration Profile Win64" Requires="Prerequisites Agent Win64;Automated Test Agent $(ProjectName) Win64" UgsProject="$(UgsProject)" UgsBadge="$(UgsBadge)"/>
+		<Label Category="Test" Name="$(ProjectName) Iteration Profile Mac" Requires="Prerequisites Agent Mac;Automated Test Agent $(ProjectName) Mac" UgsProject="$(UgsProject)" UgsBadge="$(UgsBadge)"/>
 		
 	</ForEach>
 
