IterationProfile: Split Code and Warm into their own Node so they are easier to eyeball the exec times of in the Horde web client.

This was tested via a local IterationProfile run.

diff --git a/Engine/Build/IterationProfile.xml b/Engine/Build/IterationProfile.xml
--- a/Engine/Build/IterationProfile.xml
+++ b/Engine/Build/IterationProfile.xml
@@ -168,15 +168,20 @@
 			<Property Name="BaselineNodeRequires" Value="$(RunNodeName)" If="$(RunAndBaseline)"/>
 			<Do If="$(Win64)">
 				<!-- Run the Tests -->
-				<Node Name="$(RunNodeName)" Requires="$(RunNodeRequires)">
-					<Do If="$(ColdPass)">
+				<Do If="$(ColdPass)">
+					<Property Name="ColdNodeName" Value="$(RunNodeName) Cold"/>
+					<Node Name="$(ColdNodeName)" Requires="$(RunNodeRequires)">
 						<Expand Name="RunTest" OutputName="$(OutputName)" PlatformName="$(PlatformName)" ProjectName="$(ProjectName)" NodeName="$(RunNodeName)" Phase="Cold" Args="$(ColdArgs)"/>
-					</Do>
-					<Do If="$(WarmPass)">
+					</Node>
+					<Property Name="RunGraphNodes" Value="$(RunGraphNodes);$(ColdNodeName)"/>
+				</Do>
+				<Do If="$(WarmPass)">
+					<Property Name="WarmNodeName" Value="$(RunNodeName) Warm"/>
+					<Node Name="$(WarmNodeName)" Requires="$(RunNodeRequires)">
 						<Expand Name="RunTest" OutputName="$(OutputName)" PlatformName="$(PlatformName)" ProjectName="$(ProjectName)" NodeName="$(RunNodeName)" Phase="Warm" Args="$(WarmArgs)"/>
-					</Do>
-				</Node>
-				<Property Name="RunGraphNodes" Value="$(RunGraphNodes);$(RunNodeName)"/>
+					</Node>
+					<Property Name="RunGraphNodes" Value="$(RunGraphNodes);$(WarmNodeName)"/>
+				</Do>
 				<!-- Baseline the Results -->
 				<Node Name="$(BaselineNodeName)" Requires="$(BaselineNodeRequires)">
 					<Do If="$(ColdPass)">
@@ -206,15 +211,20 @@
 			<Property Name="BaselineNodeRequires" Value="$(RunNodeName)" If="$(RunAndBaseline)"/>
 			<Do If="$(MacOS)">
 				<!-- Run the Tests -->
-				<Node Name="$(RunNodeName)" Requires="$(RunNodeRequires)">
-					<Do If="$(ColdPass)">
+				<Do If="$(ColdPass)">
+					<Property Name="ColdNodeName" Value="$(RunNodeName) Cold"/>
+					<Node Name="$(ColdNodeName)" Requires="$(RunNodeRequires)">
 						<Expand Name="RunTest" OutputName="$(OutputName)" PlatformName="$(PlatformName)" ProjectName="$(ProjectName)" NodeName="$(RunNodeName)" Phase="Cold" Args="$(ColdArgs)"/>
-					</Do>
-					<Do If="$(WarmPass)">
+					</Node>
+					<Property Name="RunGraphNodes" Value="$(RunGraphNodes);$(ColdNodeName)"/>
+				</Do>
+				<Do If="$(WarmPass)">
+					<Property Name="WarmNodeName" Value="$(RunNodeName) Warm"/>
+					<Node Name="$(WarmNodeName)" Requires="$(RunNodeRequires)">
 						<Expand Name="RunTest" OutputName="$(OutputName)" PlatformName="$(PlatformName)" ProjectName="$(ProjectName)" NodeName="$(RunNodeName)" Phase="Warm" Args="$(WarmArgs)"/>
-					</Do>
-				</Node>
-				<Property Name="RunGraphNodes" Value="$(RunGraphNodes);$(RunNodeName)"/>
+					</Node>
+					<Property Name="RunGraphNodes" Value="$(RunGraphNodes);$(WarmNodeName)"/>
+				</Do>
 				<!-- Baseline the Results -->
 				<Node Name="$(BaselineNodeName)" Requires="$(BaselineNodeRequires)">
 					<Do If="$(ColdPass)">
