IterationProfile: Increase maxduration to 6 hours

diff --git a/Engine/Build/IterationProfile.xml b/Engine/Build/IterationProfile.xml
--- a/Engine/Build/IterationProfile.xml
+++ b/Engine/Build/IterationProfile.xml
@@ -62,7 +62,7 @@
 	<!-- Gauntlet arguments for all runs -->
 	<Property Name="GauntletArgs" Value="-build=editor -UseEditor -verbose -report -tempdir=&quot;$(RootDir)/Tests&quot;"/>
 	<Property Name="GauntletArgs" Value="$(GauntletArgs) -branch=$(Branch) -ECBranch=&quot;$(Branch)&quot; -ECChangelist=&quot;$(Change)&quot;"/>
-	<Property Name="GauntletArgs" Value="$(GauntletArgs) -maxduration=14400"/>
+	<Property Name="GauntletArgs" Value="$(GauntletArgs) -maxduration=21600"/>
 
 	<!-- Gauntlet arguments for profile runs -->
 	<Property Name="ProfileArgs" Value="-test=EditorTests -testfilter=Group:Iteration -tracefiletrunc"/>
