IterationProfile: Add dependency of Cold onto Warm, only if Cold is run.

diff --git a/Engine/Build/IterationProfile.xml b/Engine/Build/IterationProfile.xml
--- a/Engine/Build/IterationProfile.xml
+++ b/Engine/Build/IterationProfile.xml
@@ -168,6 +168,7 @@
 			<Property Name="BaselineNodeRequires" Value="$(RunNodeName)" If="$(RunAndBaseline)"/>
 			<Do If="$(Win64)">
 				<!-- Run the Tests -->
+				<Property Name="ColdNodeName" Value=""/>
 				<Do If="$(ColdPass)">
 					<Property Name="ColdNodeName" Value="$(RunNodeName) Cold"/>
 					<Node Name="$(ColdNodeName)" Requires="$(RunNodeRequires)">
@@ -177,7 +178,7 @@
 				</Do>
 				<Do If="$(WarmPass)">
 					<Property Name="WarmNodeName" Value="$(RunNodeName) Warm"/>
-					<Node Name="$(WarmNodeName)" Requires="$(RunNodeRequires)">
+					<Node Name="$(WarmNodeName)" Requires="$(RunNodeRequires);$(ColdNodeName)">
 						<Expand Name="RunTest" OutputName="$(OutputName)" PlatformName="$(PlatformName)" ProjectName="$(ProjectName)" NodeName="$(RunNodeName)" Phase="Warm" Args="$(WarmArgs)"/>
 					</Node>
 					<Property Name="RunGraphNodes" Value="$(RunGraphNodes);$(WarmNodeName)"/>
@@ -211,6 +212,7 @@
 			<Property Name="BaselineNodeRequires" Value="$(RunNodeName)" If="$(RunAndBaseline)"/>
 			<Do If="$(MacOS)">
 				<!-- Run the Tests -->
+				<Property Name="ColdNodeName" Value=""/>
 				<Do If="$(ColdPass)">
 					<Property Name="ColdNodeName" Value="$(RunNodeName) Cold"/>
 					<Node Name="$(ColdNodeName)" Requires="$(RunNodeRequires)">
@@ -220,7 +222,7 @@
 				</Do>
 				<Do If="$(WarmPass)">
 					<Property Name="WarmNodeName" Value="$(RunNodeName) Warm"/>
-					<Node Name="$(WarmNodeName)" Requires="$(RunNodeRequires)">
+					<Node Name="$(WarmNodeName)" Requires="$(RunNodeRequires);$(ColdNodeName)">
 						<Expand Name="RunTest" OutputName="$(OutputName)" PlatformName="$(PlatformName)" ProjectName="$(ProjectName)" NodeName="$(RunNodeName)" Phase="Warm" Args="$(WarmArgs)"/>
 					</Node>
 					<Property Name="RunGraphNodes" Value="$(RunGraphNodes);$(WarmNodeName)"/>
