IterationProfile: Roll back into a single node to fix some errors

diff --git a/Engine/Build/IterationProfile.xml b/Engine/Build/IterationProfile.xml
--- a/Engine/Build/IterationProfile.xml
+++ b/Engine/Build/IterationProfile.xml
@@ -168,21 +168,15 @@
 			<Property Name="BaselineNodeRequires" Value="$(RunNodeName)" If="$(RunAndBaseline)"/>
 			<Do If="$(Win64)">
 				<!-- Run the Tests -->
-				<Property Name="ColdNodeName" Value=""/>
-				<Do If="$(ColdPass)">
-					<Property Name="ColdNodeName" Value="$(RunNodeName) Cold"/>
-					<Node Name="$(ColdNodeName)" Requires="$(RunNodeRequires)">
+				<Node Name="$(RunNodeName)" Requires="$(RunNodeRequires)">
+					<Do If="$(ColdPass)">
 						<Expand Name="RunTest" OutputName="$(OutputName)" PlatformName="$(PlatformName)" ProjectName="$(ProjectName)" NodeName="$(RunNodeName)" Phase="Cold" Args="$(ColdArgs)"/>
-					</Node>
-					<Property Name="RunGraphNodes" Value="$(RunGraphNodes);$(ColdNodeName)"/>
-				</Do>
-				<Do If="$(WarmPass)">
-					<Property Name="WarmNodeName" Value="$(RunNodeName) Warm"/>
-					<Node Name="$(WarmNodeName)" Requires="$(RunNodeRequires);$(ColdNodeName)">
+					</Do>
+					<Do If="$(WarmPass)">
 						<Expand Name="RunTest" OutputName="$(OutputName)" PlatformName="$(PlatformName)" ProjectName="$(ProjectName)" NodeName="$(RunNodeName)" Phase="Warm" Args="$(WarmArgs)"/>
-					</Node>
-					<Property Name="RunGraphNodes" Value="$(RunGraphNodes);$(WarmNodeName)"/>
-				</Do>
+					</Do>
+				</Node>
+				<Property Name="RunGraphNodes" Value="$(RunGraphNodes);$(RunNodeName)"/>
 				<!-- Baseline the Results -->
 				<Node Name="$(BaselineNodeName)" Requires="$(BaselineNodeRequires)">
 					<Do If="$(ColdPass)">
@@ -212,21 +206,15 @@
 			<Property Name="BaselineNodeRequires" Value="$(RunNodeName)" If="$(RunAndBaseline)"/>
 			<Do If="$(MacOS)">
 				<!-- Run the Tests -->
-				<Property Name="ColdNodeName" Value=""/>
-				<Do If="$(ColdPass)">
-					<Property Name="ColdNodeName" Value="$(RunNodeName) Cold"/>
-					<Node Name="$(ColdNodeName)" Requires="$(RunNodeRequires)">
+				<Node Name="$(RunNodeName)" Requires="$(RunNodeRequires)">
+					<Do If="$(ColdPass)">
 						<Expand Name="RunTest" OutputName="$(OutputName)" PlatformName="$(PlatformName)" ProjectName="$(ProjectName)" NodeName="$(RunNodeName)" Phase="Cold" Args="$(ColdArgs)"/>
-					</Node>
-					<Property Name="RunGraphNodes" Value="$(RunGraphNodes);$(ColdNodeName)"/>
-				</Do>
-				<Do If="$(WarmPass)">
-					<Property Name="WarmNodeName" Value="$(RunNodeName) Warm"/>
-					<Node Name="$(WarmNodeName)" Requires="$(RunNodeRequires);$(ColdNodeName)">
+					</Do>
+					<Do If="$(WarmPass)">
 						<Expand Name="RunTest" OutputName="$(OutputName)" PlatformName="$(PlatformName)" ProjectName="$(ProjectName)" NodeName="$(RunNodeName)" Phase="Warm" Args="$(WarmArgs)"/>
-					</Node>
-					<Property Name="RunGraphNodes" Value="$(RunGraphNodes);$(WarmNodeName)"/>
-				</Do>
+					</Do>
+				</Node>
+				<Property Name="RunGraphNodes" Value="$(RunGraphNodes);$(RunNodeName)"/>
 				<!-- Baseline the Results -->
 				<Node Name="$(BaselineNodeName)" Requires="$(BaselineNodeRequires)">
 					<Do If="$(ColdPass)">
