IterationProfile: Bump timeout up to 8 hours as 6 hours is right on the edge of what a 6-core Xeon Gold can do w/ Reverb.

diff --git a/Engine/Build/IterationProfile.xml b/Engine/Build/IterationProfile.xml
--- a/Engine/Build/IterationProfile.xml
+++ b/Engine/Build/IterationProfile.xml
@@ -62,7 +62,7 @@
 	<!-- Gauntlet arguments for all runs -->
 	<Property Name="GauntletArgs" Value="-build=editor -UseEditor -verbose -report -tempdir=&quot;$(RootDir)/Tests&quot;"/>
 	<Property Name="GauntletArgs" Value="$(GauntletArgs) -branch=$(Branch) -ECBranch=&quot;$(Branch)&quot; -ECChangelist=&quot;$(Change)&quot;"/>
-	<Property Name="GauntletArgs" Value="$(GauntletArgs) -maxduration=21600"/>
+	<Property Name="GauntletArgs" Value="$(GauntletArgs) -maxduration=28800"/>
 
 	<!-- Gauntlet arguments for profile runs -->
 	<Property Name="ProfileArgs" Value="-test=EditorTests -testfilter=Group:Iteration -tracefiletrunc"/>
