Setup separate agent types for each project running IterationProfile so they always run on the same machine.

Create a separate agent type for each project we run iteration profile on
Reference a new pool in each agent type that should have just a single worker in it
Update "IterationTest" -> "IterationProfile" for consistency

diff --git a/Engine/Build/IterationProfile.xml b/Engine/Build/IterationProfile.xml
--- a/Engine/Build/IterationProfile.xml
+++ b/Engine/Build/IterationProfile.xml
@@ -158,7 +158,7 @@
 				<Compile Target="$(ProjectName)Editor" Platform="Win64" Configuration="Development" Arguments="$(EditorCompileArgs)"/>
 			</Node>
 		</Agent>
-		<Agent Name="Automated Test Agent $(ProjectName) Win64" Type="IterationTest">
+		<Agent Name="Automated Test Agent $(ProjectName) Win64" Type="IterationProfile$(ProjectName)">
 			<Property Name="PlatformName" Value="Win64"/>
 			<Property Name="OutputName" Value="$(ProjectName)$(PlatformName)IterationProfile"/>
 			<Property Name="RunNodeName" Value="Run $(OutputName)"/>
