IterationProfile: Add support for incremental phase

= IterationProfile.xml =

Cleanup for clarity's sake:
* Pass -> Phase
* PlatformNames -> Platforms
* PlatformName -> Platform
* ProjectNames -> Projects
* ProjectName -> Project

Rework logic to simplifiy the graph by using property-dependent property names (see pending 16866021):
* Rework Cold and Warm Phase booleans into <ForEach> of Semicolon-delimited build option: Phases
* Rework Win64 and MacOS Platform booleans into <ForEach> of Semicolon-delimited build option: Platforms

= ue5-main.json =

Split IterationProfile into two templates for speed: iteration-profile-nightly and iteration-profile-tests
* Nightly does Cold+Warm once a day
* Tests does Incremental+Warm the all day long

#rnx

diff --git a/Engine/Build/BatchFiles/IterationProfile.bat b/Engine/Build/BatchFiles/IterationProfile.bat
--- a/Engine/Build/BatchFiles/IterationProfile.bat
+++ b/Engine/Build/BatchFiles/IterationProfile.bat
@@ -1 +1 @@
-@call %~dp0RunUAT.bat BuildGraph -Script=Engine/Build/IterationProfile.xml -Target="Run Iteration Profile Tests" -append:ProjectNames="ShooterGame" -set:Win64=true %*
+@call %~dp0RunUAT.bat BuildGraph -Script=Engine/Build/IterationProfile.xml -Target="Run Iteration Profile Tests" -set:Platforms=Win64 -append:Projects="EngineTest" -append:Projects="Reverb" -append:Projects="ShooterGame" %*
diff --git a/Engine/Build/BatchFiles/IterationProfileBaseline.bat b/Engine/Build/BatchFiles/IterationProfileBaseline.bat
--- a/Engine/Build/BatchFiles/IterationProfileBaseline.bat
+++ b/Engine/Build/BatchFiles/IterationProfileBaseline.bat
@@ -1 +1 @@
-@call %~dp0RunUAT.bat BuildGraph -Script=Engine/Build/IterationProfile.xml -Target="Baseline Iteration Profile Results" -append:ProjectNames="EngineTest" -append:ProjectNames="Reverb" -append:ProjectNames="ShooterGame" -set:Win64=true %*
+@call %~dp0RunUAT.bat BuildGraph -Script=Engine/Build/IterationProfile.xml -Target="Baseline Iteration Profile Results" -set:Platforms=Win64 -append:Projects="EngineTest" -append:Projects="Reverb" -append:Projects="ShooterGame" %*
diff --git a/Engine/Build/BatchFiles/IterationProfileRunAndBaseline.bat b/Engine/Build/BatchFiles/IterationProfileRunAndBaseline.bat
deleted file mode 100755
diff --git a/Engine/Build/IterationProfile.xml b/Engine/Build/IterationProfile.xml
--- a/Engine/Build/IterationProfile.xml
+++ b/Engine/Build/IterationProfile.xml
@@ -13,15 +13,14 @@
 	<Property Name="BuildName" Value="$(EscapedBranch)-CL-$(Change)$(PreflightSuffix)"/>
 
 	<!-- Projects setup-->
-	<Option Name="ProjectNames" DefaultValue="false" Restrict=".+" Description="The Projects to test"/>
-	<Error Message="No project names specified" If="'$(ProjectNames)' == 'false'"/>
-	<Option Name="UgsProject" DefaultValue="" Description="The relative path to the project"/>
-	<Option Name="UgsBadge" DefaultValue="" Description="The badge name to use with labels"/>
+	<Option Name="Projects" DefaultValue="false" Restrict=".+" Description="The projects to test"/>
+	<Error Message="No project names specified" If="'$(Projects)' == 'false'"/>
 
 	<!-- Platform Options -->
-	<Option Name="Win64" DefaultValue="false" Description="Enable Win64 Platform"/>
-	<Option Name="MacOS" DefaultValue="false" Description="Enable MacOS Platform"/>
-	<Error Message="No platforms specified" If="!$(Win64) And !$(MacOS)"/>
+	<Option Name="Platforms" DefaultValue="Win64" Description="The list of platforms to profile"/>
+	<Property Name="Win64CompileAgentType" Value="Win64"/>
+	<Property Name="MacCompileAgentType" Value="TestMac"/>
+	<Error Message="No platforms specified" If="'$(Platforms)' == ''"/>
 
 	<!-- Build Options -->
 	<Property Name="ExtraPluginsCommaDelims" Value="EditorTests"/>
@@ -53,10 +52,11 @@
 	<Property Name="BaselineAutomationNetworkOutputDir" Value="$(AutomationOutputDir)" If="$(RunAndBaseline)"/>
 
 	<!-- Arguments for tracing -->
-	<Option Name="ColdPass" DefaultValue="true" Description="Perform cold workspace loading profiles"/>
-	<Option Name="ColdArgs" DefaultValue="-ddc=Cold" Description="Arguments to unreal for the cold pass"/>
-	<Option Name="WarmPass" DefaultValue="true" Description="Perform warm (incremental) workspace loading profiles"/>
-	<Option Name="WarmArgs" DefaultValue="-ddc=Warm" Description="Arguments to unreal for the warm pass"/>
+	<Option Name="Phases" DefaultValue="Incr;Warm" Description="The list of phases to perform"/>
+	<Option Name="ColdArgs" DefaultValue="-ddc=Cold" Description="Arguments to unreal for the cold phase"/>
+	<Option Name="IncrArgs" DefaultValue="-ddc=Warm" Description="Arguments to unreal for the incremental phase"/>
+	<Option Name="WarmArgs" DefaultValue="-ddc=Warm" Description="Arguments to unreal for the warm phase"/>
+	<Error Message="No phases specified" If="'$(Phases)' == ''"/>
 	<Property Name="TraceDir" Value="$(RootDir)/Engine/Programs/UnrealInsights/Saved/TraceSessions"/>
 
 	<!-- Gauntlet arguments for all runs -->
@@ -77,10 +77,10 @@
 	<Property Name="SummarizeArgs" Value="$(SummarizeArgs) -skipbaseline" If="$(RunAndBaseline)"/>
 
 	<!-- The main test running macro logic that is expanded with platform and dependency node names -->
-	<Macro Name="RunTest" Arguments="OutputName;PlatformName;ProjectName;NodeName;Phase;Args">
+	<Macro Name="RunTest" Arguments="OutputName;Platform;Project;NodeName;Phase;Args">
 		<Property Name="TestName" Value="IterationProfile$(Phase)"/>
-		<Property Name="TraceFilePath" Value="$(TraceDir)/$(ProjectName)$(PlatformName)$(TestName).utrace"/>
-		<Property Name="TelemetryFileName" Value="$(ProjectName)$(PlatformName)$(TestName)Telemetry.csv"/>
+		<Property Name="TraceFilePath" Value="$(TraceDir)/$(Project)$(Platform)$(TestName).utrace"/>
+		<Property Name="TelemetryFileName" Value="$(Project)$(Platform)$(TestName)Telemetry.csv"/>
 		<Property Name="TelemetryFilePath" Value="$(TraceDir)/$(TelemetryFileName)"/>
 		<Property Name="ReportExportPath" Value="$(AutomationOutputDir)/$(OutputName)"/>
 		<Property Name="ReportParams" Value="-ReportExportPath=&quot;$(ReportExportPath)&quot;"/>
@@ -89,7 +89,7 @@
 			<Property Name="ReportParams" Value="$(ReportParams) -ReportURL=&quot;$(ReportExportURL)&quot; -WriteTestResultsForHorde -HordeTestDataKey=&quot;$(NodeName) $(Phase)$(PreflightSuffix)&quot;"/>
 		</Do>
 		<Property Name="LogOutputParams" Value="-logdir=&quot;$(AutomationOutputDir)/$(OutputName)&quot;"/>
-		<Property Name="AllGauntletArgs" Value="-project=$(ProjectName) -platform=$(PlatformName) $(GauntletArgs) $(ReportParams) $(LogOutputParams) -artifactname=&quot;$(Phase)&quot;"/>
+		<Property Name="AllGauntletArgs" Value="-project=$(Project) -platform=$(Platform) $(GauntletArgs) $(ReportParams) $(LogOutputParams) -artifactname=&quot;$(Phase)&quot;"/>
 
 		<!-- Run the Gauntlet Iteration Profile Tests -->
 		<Command Name="RunUnreal" Arguments="$(AllGauntletArgs) $(ProfileArgs) $(Args) -tracefile=&quot;$(TraceFilePath)&quot;"/>
@@ -102,11 +102,11 @@
 
 		<Do If="$(IsBuildMachine)">
 			<!-- Files copied to this folder are automatically uploaded to Horde as artifacts -->
-			<Copy From="$(TraceDir)/$(ProjectName)$(PlatformName)$(TestName)*.*" To="$(RootDir)/Engine/Programs/AutomationTool/Saved/Logs"/>
+			<Copy From="$(TraceDir)/$(Project)$(Platform)$(TestName)*.*" To="$(RootDir)/Engine/Programs/AutomationTool/Saved/Logs"/>
 
 			<Do If="!$(IsPreflight)">
 				<!-- Publish telemetry to the db -->
-				<Command Name="PublishUnrealAutomationTelemetry" Arguments="-CSVFile=&quot;$(TelemetryFilePath)&quot; -project=$(ProjectName) -changelist=$(Change) -branch=$(Branch)"/>
+				<Command Name="PublishUnrealAutomationTelemetry" Arguments="-CSVFile=&quot;$(TelemetryFilePath)&quot; -project=$(Project) -changelist=$(Change) -branch=$(Branch)"/>
 			</Do>
 		</Do>
 	</Macro>
@@ -115,9 +115,9 @@
 	<Property Name="RunGraphNodes" Value=""/>
 
 	<!-- Mimic of the Runner that retrieves the stashed telemetry file and copies it to the depot location -->
-	<Macro Name="BaselineResult" Arguments="PlatformName;ProjectName;Phase">
+	<Macro Name="BaselineResult" Arguments="Platform;Project;Phase">
 		<Property Name="TestName" Value="IterationProfile$(Phase)"/>
-		<Property Name="TelemetryFileName" Value="$(ProjectName)$(PlatformName)$(TestName)Telemetry.csv"/>
+		<Property Name="TelemetryFileName" Value="$(Project)$(Platform)$(TestName)Telemetry.csv"/>
 		<Property Name="LocalTelemetryFilePath" Value="$(BaselineDir)/$(TelemetryFileName)"/>
 		<Property Name="NetworkTelemetryFilePath" Value="$(BaselineAutomationNetworkOutputDir)/$(TelemetryFileName)"/>
 		<Copy From="$(NetworkTelemetryFilePath)" To="$(LocalTelemetryFilePath)" Overwrite="true"/>
@@ -129,110 +129,59 @@
 	<!-- This will be set to the names of all nodes we want to execute based on specified platforms/options -->
 	<Property Name="BaselineGraphNodes" Value=""/>
 
-	<!-- Windows -->
-	<Agent Name="Prerequisites Agent Win64" Type="Win64">
-		<Node Name="Compile Prerequisites Win64Host">
-			<Compile Target="UnrealHeaderTool" Platform="Win64" Configuration="Development"/>
-			<Compile Target="ShaderCompileWorker" Platform="Win64" Configuration="Development"/>
-			<Compile Target="CrashReportClient" Platform="Win64" Configuration="Shipping"/>
-			<Compile Target="CrashReportClientEditor" Platform="Win64" Configuration="Shipping"/>
-		</Node>
-	</Agent>
+	<!-- ForEach Platform to run, generate graph nodes -->
+	<ForEach Name="Platform" Values="$(Platforms)">
 
-	<!-- macOS -->
-	<Agent Name="Prerequisites Agent Mac" Type="TestMac">
-		<Node Name="Compile Prerequisites Mac">
-			<Compile Target="UnrealHeaderTool" Platform="Mac" Configuration="Development"/>
-			<Compile Target="ShaderCompileWorker" Platform="Mac" Configuration="Development"/>
-			<Compile Target="CrashReportClient" Platform="Mac" Configuration="Shipping"/>
-			<Compile Target="CrashReportClientEditor" Platform="Mac" Configuration="Shipping"/>
-		</Node>
-	</Agent>
-	
-	<!-- ForEach Project to run, generate the graph nodes -->
-	<ForEach Name="ProjectName" Values="$(ProjectNames)">
-
-		<!-- Windows -->
-		<Agent Name="Prerequisites Agent Win64" Type="Win64">
-			<Node Name="Compile $(ProjectName) Editor Win64" Requires="Compile Prerequisites Win64Host">
-				<Compile Target="$(ProjectName)Editor" Platform="Win64" Configuration="Development" Arguments="$(EditorCompileArgs)"/>
+		<!-- Project-independent prereq compiles -->
+		<Agent Name="Prerequisites Agent $(Platform)" Type="$($(Platform)CompileAgentType)">
+			<Node Name="Compile Prerequisites $(Platform)">
+				<Compile Target="UnrealHeaderTool" Platform="$(Platform)" Configuration="Development"/>
+				<Compile Target="ShaderCompileWorker" Platform="$(Platform)" Configuration="Development"/>
+				<Compile Target="CrashReportClient" Platform="$(Platform)" Configuration="Shipping"/>
+				<Compile Target="CrashReportClientEditor" Platform="$(Platform)" Configuration="Shipping"/>
 			</Node>
-		</Agent>
-		<Agent Name="Automated Test Agent $(ProjectName) Win64" Type="IterationProfile$(ProjectName)">
-			<Property Name="PlatformName" Value="Win64"/>
-			<Property Name="OutputName" Value="$(ProjectName)$(PlatformName)IterationProfile"/>
-			<Property Name="RunNodeName" Value="Run $(OutputName)"/>
-			<Property Name="RunNodeRequires" Value="Compile Prerequisites $(PlatformName)Host;Compile $(ProjectName) Editor $(PlatformName)"/>
-			<Property Name="BaselineNodeName" Value="Baseline $(OutputName)"/>
-			<Property Name="BaselineNodeRequires" Value=""/>
-			<Property Name="BaselineNodeRequires" Value="$(RunNodeName)" If="$(RunAndBaseline)"/>
-			<Do If="$(Win64)">
-				<!-- Run the Tests -->
-				<Node Name="$(RunNodeName)" Requires="$(RunNodeRequires)">
-					<Do If="$(ColdPass)">
-						<Expand Name="RunTest" OutputName="$(OutputName)" PlatformName="$(PlatformName)" ProjectName="$(ProjectName)" NodeName="$(RunNodeName)" Phase="Cold" Args="$(ColdArgs)"/>
-					</Do>
-					<Do If="$(WarmPass)">
-						<Expand Name="RunTest" OutputName="$(OutputName)" PlatformName="$(PlatformName)" ProjectName="$(ProjectName)" NodeName="$(RunNodeName)" Phase="Warm" Args="$(WarmArgs)"/>
-					</Do>
-				</Node>
-				<Property Name="RunGraphNodes" Value="$(RunGraphNodes);$(RunNodeName)"/>
-				<!-- Baseline the Results -->
-				<Node Name="$(BaselineNodeName)" Requires="$(BaselineNodeRequires)">
-					<Do If="$(ColdPass)">
-						<Expand Name="BaselineResult" PlatformName="$(PlatformName)" ProjectName="$(ProjectName)" Phase="Cold"/>
-					</Do>
-					<Do If="$(WarmPass)">
-						<Expand Name="BaselineResult" PlatformName="$(PlatformName)" ProjectName="$(ProjectName)" Phase="Warm"/>
-					</Do>
+
+			<!-- ForEach Project to run, generate the compile graph nodes -->
+			<ForEach Name="Project" Values="$(Projects)">
+				<Node Name="Compile $(Project) Editor $(Platform)" Requires="Compile Prerequisites $(Platform)">
+					<Compile Target="$(Project)Editor" Platform="$(Platform)" Configuration="Development" Arguments="$(EditorCompileArgs)"/>
 				</Node>
-				<Property Name="BaselineGraphNodes" Value="$(BaselineGraphNodes);$(BaselineNodeName)"/>
-			</Do>
-		</Agent>
-		
-		<!-- macOS -->
-		<Agent Name="Prerequisites Agent Mac" Type="TestMac">
-			<Node Name="Compile $(ProjectName) Editor Mac" Requires="Compile Prerequisites Mac">
-				<Compile Target="$(ProjectName)Editor" Platform="Mac" Configuration="Development"  Arguments="$(EditorCompileArgs)"/>
-			</Node>
+			</ForEach>
 		</Agent>
-		<Agent Name="Automated Test Agent $(ProjectName) Mac" Type="TestMac">
-			<Property Name="PlatformName" Value="Mac"/>
-			<Property Name="OutputName" Value="$(ProjectName)$(PlatformName)IterationProfile"/>
-			<Property Name="RunNodeName" Value="Run $(OutputName)"/>
-			<Property Name="RunNodeRequires" Value="Compile Prerequisites $(PlatformName);Compile $(ProjectName) Editor $(PlatformName)"/>
-			<Property Name="BaselineNodeName" Value="Baseline $(OutputName)"/>
-			<Property Name="BaselineNodeRequires" Value=""/>
-			<Property Name="BaselineNodeRequires" Value="$(RunNodeName)" If="$(RunAndBaseline)"/>
-			<Do If="$(MacOS)">
+
+		<!-- ForEach Project to run, generate the graph nodes -->
+		<ForEach Name="Project" Values="$(Projects)">
+
+			<!-- Project execution -->
+			<Agent Name="Automated Test Agent $(Project) $(Platform)" Type="IterationProfile$(Platform)$(Project)">
+				<Property Name="OutputName" Value="$(Project)$(Platform)IterationProfile"/>
+				<Property Name="RunNodeName" Value="Run $(OutputName)"/>
+				<Property Name="RunNodeRequires" Value="Compile Prerequisites $(Platform);Compile $(Project) Editor $(Platform)"/>
+				<Property Name="BaselineNodeName" Value="Baseline $(OutputName)"/>
+				<Property Name="BaselineNodeRequires" Value=""/>
+				<Property Name="BaselineNodeRequires" Value="$(RunNodeName)" If="$(RunAndBaseline)"/>
 				<!-- Run the Tests -->
 				<Node Name="$(RunNodeName)" Requires="$(RunNodeRequires)">
-					<Do If="$(ColdPass)">
-						<Expand Name="RunTest" OutputName="$(OutputName)" PlatformName="$(PlatformName)" ProjectName="$(ProjectName)" NodeName="$(RunNodeName)" Phase="Cold" Args="$(ColdArgs)"/>
-					</Do>
-					<Do If="$(WarmPass)">
-						<Expand Name="RunTest" OutputName="$(OutputName)" PlatformName="$(PlatformName)" ProjectName="$(ProjectName)" NodeName="$(RunNodeName)" Phase="Warm" Args="$(WarmArgs)"/>
-					</Do>
+					<ForEach Name="Phase" Values="$(Phases)">
+						<Expand Name="RunTest" OutputName="$(OutputName)" Platform="$(Platform)" Project="$(Project)" NodeName="$(RunNodeName)" Phase="$(Phase)" Args="$($(Phase)Args)"/>
+					</ForEach>
 				</Node>
 				<Property Name="RunGraphNodes" Value="$(RunGraphNodes);$(RunNodeName)"/>
 				<!-- Baseline the Results -->
 				<Node Name="$(BaselineNodeName)" Requires="$(BaselineNodeRequires)">
-					<Do If="$(ColdPass)">
-						<Expand Name="BaselineResult" PlatformName="$(PlatformName)" ProjectName="$(ProjectName)" Phase="Cold"/>
-					</Do>
-					<Do If="$(WarmPass)">
-						<Expand Name="BaselineResult" PlatformName="$(PlatformName)" ProjectName="$(ProjectName)" Phase="Warm"/>
-					</Do>
+					<ForEach Name="Phase" Values="$(Phases)">
+						<Expand Name="BaselineResult" Platform="$(Platform)" Project="$(Project)" Phase="$(Phase)"/>
+					</ForEach>
 				</Node>
 				<Property Name="BaselineGraphNodes" Value="$(BaselineGraphNodes);$(BaselineNodeName)"/>
-			</Do>
-		</Agent>
+			</Agent>
 
-		<!-- Graph nodes that really tie the room together -->
-		<Label Category="Test" Name="$(ProjectName) Iteration Profile Win64" Requires="Prerequisites Agent Win64;Automated Test Agent $(ProjectName) Win64" UgsProject="$(UgsProject)" UgsBadge="$(UgsBadge)"/>
-		<Label Category="Test" Name="$(ProjectName) Iteration Profile Mac" Requires="Prerequisites Agent Mac;Automated Test Agent $(ProjectName) Mac" UgsProject="$(UgsProject)" UgsBadge="$(UgsBadge)"/>
-		
-	</ForEach>
+			<!-- Graph nodes that really tie the room together -->
+			<Label Category="Test" Name="$(Project) Iteration Profile $(Platform)" Requires="Prerequisites Agent $(Platform);Automated Test Agent $(Project) $(Platform)"/>
+			
+		</ForEach> <!-- Projects -->
+
+	</ForEach> <!-- Platforms -->
 
 	<!-- These Nodes run the Editor and gather profile data -->
 	<Aggregate Name="Run Iteration Profile Tests" Requires="$(RunGraphNodes)"/>
