Remove Lumin as build target / disable by default.

diff --git a/Engine/Build/InstalledEngineBuild.xml b/Engine/Build/InstalledEngineBuild.xml
--- a/Engine/Build/InstalledEngineBuild.xml
+++ b/Engine/Build/InstalledEngineBuild.xml
@@ -42,8 +42,8 @@
 	<Option Name="WithTVOS" Restrict="true|false" DefaultValue="$(DefaultWithIOS)" Description="Include the tvOS target platform"/>
 	<Option Name="WithLinux" Restrict="true|false" DefaultValue="$(DefaultWithLinux)" Description="Include the Linux target platform"/>
 	<Option Name="WithLinuxArm64" Restrict="true|false" DefaultValue="$(DefaultWithLinuxArm64)" Description="Include the Linux AArch64 target platform"/>
-	<Option Name="WithLumin" Restrict="true|false" DefaultValue="$(DefaultWithPlatform)" Description="Include the Lumin target platform on Windows"/>
-	<Option Name="WithLuminMac" Restrict="true|false" DefaultValue="$(DefaultWithPlatform)" Description="Include the Lumin target platform on Mac"/>
+	<Option Name="WithLumin" Restrict="true|false" DefaultValue="false" Description="Include the Lumin target platform on Windows"/>
+	<Option Name="WithLuminMac" Restrict="true|false" DefaultValue="false" Description="Include the Lumin target platform on Mac"/>
 	<Option Name="WithHoloLens" Restrict="true|false" DefaultValue="$(DefaultWithHoloLens)" Description="Include the HoloLens target platform"/>
 
 	<!-- Set up target types -->
