IterationProfile: Disable trace server via cmdline arg

This change eliminates an intermittent warning print that only seems to happen on the SummarizeTraceCommandlet, and we don't need the trace server for that commandlet to do it's job.

diff --git a/Engine/Build/IterationProfile.xml b/Engine/Build/IterationProfile.xml
--- a/Engine/Build/IterationProfile.xml
+++ b/Engine/Build/IterationProfile.xml
@@ -73,7 +73,7 @@
 	<Property Name="ProfileArgs" Value="$(ProfileArgs) -NoShaderDistrib -VerboseShaderLogging"/>
 
 	<!-- Gauntlet arguments for summarize commandlet runs -->
-	<Property Name="SummarizeArgs" Value="-statsfile=&quot;$(RootDir)/Engine/Build/IterationProfileStats.csv&quot;"/>
+	<Property Name="SummarizeArgs" Value="-statsfile=&quot;$(RootDir)/Engine/Build/IterationProfileStats.csv&quot; -notraceserver"/>
 	<Property Name="SummarizeArgs" Value="$(SummarizeArgs) -skipbaseline" If="$(RunAndBaseline)"/>
 
 	<!-- The main test running macro logic that is expanded with platform and dependency node names -->
diff --git a/Engine/Source/Runtime/Core/Private/ProfilingDebugging/TraceAuxiliary.cpp b/Engine/Source/Runtime/Core/Private/ProfilingDebugging/TraceAuxiliary.cpp
--- a/Engine/Source/Runtime/Core/Private/ProfilingDebugging/TraceAuxiliary.cpp
+++ b/Engine/Source/Runtime/Core/Private/ProfilingDebugging/TraceAuxiliary.cpp
@@ -571,6 +571,7 @@ void FTraceAuxiliary::Initialize(const TCHAR* CommandLine)
 	TRACE_CPUPROFILER_EVENT_SCOPE(FTraceAux_Init);
 
 #if WITH_UNREAL_TRACE_LAUNCH
+	if (!FParse::Param(CommandLine, TEXT("notraceserver")))
 	{
 		TRACE_CPUPROFILER_EVENT_SCOPE(FTraceAux_LaunchUnrealTrace);
 		LaunchUnrealTraceInternal();
