IterationProfile: Post trace and telemetry artifacts in the network folder, not the artifact store

Larger projects have utrace files north of 1GB, so they don't stuff into mongodb well.
To make these artifacts accessible, we need to provide a properly-slashed path, so flip slashes to windows styles (for windows builds at the very least)

diff --git a/Engine/Build/IterationProfile.xml b/Engine/Build/IterationProfile.xml
--- a/Engine/Build/IterationProfile.xml
+++ b/Engine/Build/IterationProfile.xml
@@ -33,13 +33,14 @@
 	<Property Name="EditorCompileArgs" Value="-BuildPlugin=$(ExtraPluginsPlusDelims)"/>
 
 	<!-- Automation uses its own network output directory -->
-	<Property Name="AutomationNetworkRootDir" Value="//epicgames.net/root/Builds/Automation/Reports/IterationProfile"/>
+	<Property Name="AutomationNetworkRootDir" Value="\\epicgames.net\root\Builds\Automation\Reports\IterationProfile"/>
 	<Property Name="AutomationNetworkRootDir" Value="/Volumes/Builds/Automation/Reports/IterationProfile" If="'$(HostPlatform)' != 'Win64'"/>
 	<Property Name="AutomationReportRootUrl" Value="http://automation.epicgames.net/reports/iterationprofile"/>
 	<Property Name="AutomationLocalOutputDir" Value="$(RootDir)/LocalBuilds/Reports"/>
 
-	<!-- Outputs for this run -->
-	<Property Name="AutomationNetworkOutputDir" Value="$(AutomationNetworkRootDir)/$(BuildName)"/>
+	<!-- Outputs for this run keeping slashes flipped the right way because it shows up on the webapp this way -->
+	<Property Name="AutomationNetworkOutputDir" Value="$(AutomationNetworkRootDir)\$(BuildName)"/>
+	<Property Name="AutomationNetworkOutputDir" Value="$(AutomationNetworkRootDir)/$(BuildName)" If="'$(HostPlatform)' != 'Win64'"/>
 	<Property Name="AutomationReportBaseUrl" Value="$(AutomationReportRootUrl)/$(BuildName)"/>
 	<Property Name="AutomationOutputDir" Value="$(AutomationLocalOutputDir)"/>
 	<Property Name="AutomationOutputDir" Value="$(AutomationNetworkOutputDir)" If="$(IsBuildMachine)"/>
@@ -82,13 +83,14 @@
 		<Property Name="TraceFilePath" Value="$(TraceDir)/$(Project)$(Platform)$(TestName).utrace"/>
 		<Property Name="TelemetryFileName" Value="$(Project)$(Platform)$(TestName)Telemetry.csv"/>
 		<Property Name="TelemetryFilePath" Value="$(TraceDir)/$(TelemetryFileName)"/>
-		<Property Name="ReportExportPath" Value="$(AutomationOutputDir)/$(OutputName)"/>
-		<Property Name="ReportParams" Value="-ReportExportPath=&quot;$(ReportExportPath)&quot;"/>
+		<Property Name="OutputPath" Value="$(AutomationOutputDir)\$(OutputName)"/>
+		<Property Name="OutputPath" Value="$(AutomationOutputDir)/$(OutputName)" If="'$(HostPlatform)' != 'Win64'"/>
+		<Property Name="ReportParams" Value="-ReportExportPath=&quot;$(OutputPath)&quot;"/>
 		<Do If="$(IsBuildMachine)">
 			<Property Name="ReportExportURL" Value="$(AutomationReportBaseUrl)/$(OutputName)/index.html"/>
 			<Property Name="ReportParams" Value="$(ReportParams) -ReportURL=&quot;$(ReportExportURL)&quot; -WriteTestResultsForHorde -HordeTestDataKey=&quot;$(NodeName) $(Phase)$(PreflightSuffix)&quot;"/>
 		</Do>
-		<Property Name="LogOutputParams" Value="-logdir=&quot;$(AutomationOutputDir)/$(OutputName)&quot;"/>
+		<Property Name="LogOutputParams" Value="-logdir=&quot;$(OutputPath)&quot;"/> <!-- handled with care, using backslashes, so it's a valid path in the Horde web pages -->
 		<Property Name="AllGauntletArgs" Value="-project=$(Project) -platform=$(Platform) $(GauntletArgs) $(ReportParams) $(LogOutputParams) -artifactname=&quot;$(Phase)&quot;"/>
 
 		<!-- Run the Gauntlet Iteration Profile Tests -->
@@ -97,17 +99,16 @@
 		<!-- Translate the utrace file into csv data -->
 		<Commandlet Name="SummarizeTrace" Arguments="-NoEnginePlugins -inputfile=&quot;$(TraceFilePath)&quot; -testname=&quot;$(TestName)&quot; $(SummarizeArgs)"/>
 
-		<!-- Stash the telemetry file in case this build will become the baseline in the future -->
-		<Copy From="$(TelemetryFilePath)" To="$(AutomationOutputDir)/$(TelemetryFileName)"/>
+		<!-- Save all files named for this run in the trace directory -->
+		<Copy From="$(TraceDir)/$(Project)$(Platform)$(TestName)*.*" To="$(OutputPath)"/>
 
 		<Do If="$(IsBuildMachine)">
-			<!-- Files copied to this folder are automatically uploaded to Horde as artifacts -->
-			<Copy From="$(TraceDir)/$(Project)$(Platform)$(TestName)*.*" To="$(RootDir)/Engine/Programs/AutomationTool/Saved/Logs"/>
 
 			<Do If="!$(IsPreflight)">
 				<!-- Publish telemetry to the db -->
 				<Command Name="PublishUnrealAutomationTelemetry" Arguments="-CSVFile=&quot;$(TelemetryFilePath)&quot; -project=$(Project) -changelist=$(Change) -branch=$(Branch)"/>
 			</Do>
+
 		</Do>
 	</Macro>
 
@@ -115,11 +116,11 @@
 	<Property Name="RunGraphNodes" Value=""/>
 
 	<!-- Mimic of the Runner that retrieves the stashed telemetry file and copies it to the depot location -->
-	<Macro Name="BaselineResult" Arguments="Platform;Project;Phase">
+	<Macro Name="BaselineResult" Arguments="OutputName;Platform;Project;Phase">
 		<Property Name="TestName" Value="IterationProfile$(Phase)"/>
 		<Property Name="TelemetryFileName" Value="$(Project)$(Platform)$(TestName)Telemetry.csv"/>
 		<Property Name="LocalTelemetryFilePath" Value="$(BaselineDir)/$(TelemetryFileName)"/>
-		<Property Name="NetworkTelemetryFilePath" Value="$(BaselineAutomationNetworkOutputDir)/$(TelemetryFileName)"/>
+		<Property Name="NetworkTelemetryFilePath" Value="$(BaselineAutomationNetworkOutputDir)/$(OutputName)/$(TelemetryFileName)"/>
 		<Copy From="$(NetworkTelemetryFilePath)" To="$(LocalTelemetryFilePath)" Overwrite="true"/>
 		<Do If="!$(IsBuildMachine)">
 			<Spawn Exe="p4" Arguments="reconcile -ea $(LocalTelemetryFilePath)"/>
@@ -160,6 +161,7 @@
 				<Property Name="BaselineNodeName" Value="Baseline $(OutputName)"/>
 				<Property Name="BaselineNodeRequires" Value=""/>
 				<Property Name="BaselineNodeRequires" Value="$(RunNodeName)" If="$(RunAndBaseline)"/>
+
 				<!-- Run the Tests -->
 				<Node Name="$(RunNodeName)" Requires="$(RunNodeRequires)">
 					<ForEach Name="Phase" Values="$(Phases)">
@@ -167,10 +169,11 @@
 					</ForEach>
 				</Node>
 				<Property Name="RunGraphNodes" Value="$(RunGraphNodes);$(RunNodeName)"/>
+
 				<!-- Baseline the Results -->
 				<Node Name="$(BaselineNodeName)" Requires="$(BaselineNodeRequires)">
 					<ForEach Name="Phase" Values="$(Phases)">
-						<Expand Name="BaselineResult" Platform="$(Platform)" Project="$(Project)" Phase="$(Phase)"/>
+						<Expand Name="BaselineResult" OutputName="$(OutputName)" Platform="$(Platform)" Project="$(Project)" Phase="$(Phase)"/>
 					</ForEach>
 				</Node>
 				<Property Name="BaselineGraphNodes" Value="$(BaselineGraphNodes);$(BaselineNodeName)"/>
