IterationProfile: Coalesce phases to incremental, warm, and cold.

#rnx

diff --git a/Engine/Build/BatchFiles/IterationProfile.bat b/Engine/Build/BatchFiles/IterationProfile.bat
--- a/Engine/Build/BatchFiles/IterationProfile.bat
+++ b/Engine/Build/BatchFiles/IterationProfile.bat
@@ -1 +1 @@
-@call %~dp0RunUAT.bat BuildGraph -Script=Engine/Build/IterationProfile.xml -Target="Run Iteration Profile Tests" -set:Platforms=Win64 -set:Phases="Cold;Warm" -append:Projects="AncientGame" -append:Projects="EngineTest" -append:Projects="ShooterGame" %*
+@call %~dp0RunUAT.bat BuildGraph -Script=Engine/Build/IterationProfile.xml -Target="Run Iteration Profile Tests" -set:Platforms=Win64 -append:Projects="AncientGame" -append:Projects="EngineTest" -append:Projects="ShooterGame" %*
\ No newline at end of file
diff --git a/Engine/Build/IterationProfile.xml b/Engine/Build/IterationProfile.xml
--- a/Engine/Build/IterationProfile.xml
+++ b/Engine/Build/IterationProfile.xml
@@ -53,7 +53,7 @@
 	<Property Name="BaselineAutomationNetworkOutputDir" Value="$(AutomationOutputDir)" If="$(RunAndBaseline)"/>
 
 	<!-- Arguments for tracing -->
-	<Option Name="Phases" DefaultValue="Incr;Warm" Description="The list of phases to perform"/>
+	<Option Name="Phases" DefaultValue="Incr;Warm;Cold" Description="The list of phases to perform"/>
 	<Option Name="ColdArgs" DefaultValue="-ddc=Cold" Description="Arguments to unreal for the cold phase"/>
 	<Option Name="IncrArgs" DefaultValue="-ddc=Warm" Description="Arguments to unreal for the incremental phase"/>
 	<Option Name="WarmArgs" DefaultValue="-ddc=Warm" Description="Arguments to unreal for the warm phase"/>
