IterationProfile: Partial back out of 16986831

Gauntlet's output folder is "tidied" up when you run it more than once, which deletes prior run's files that may have been there. This breaks copying each phase's artifacts to the output folder. For now just do what we did before: copy the artifacts to the parent directory, which Gauntlet doesn't mess with.

#rnx

diff --git a/Engine/Build/IterationProfile.xml b/Engine/Build/IterationProfile.xml
--- a/Engine/Build/IterationProfile.xml
+++ b/Engine/Build/IterationProfile.xml
@@ -83,14 +83,13 @@
 		<Property Name="TraceFilePath" Value="$(TraceDir)/$(Project)$(Platform)$(TestName).utrace"/>
 		<Property Name="TelemetryFileName" Value="$(Project)$(Platform)$(TestName)Telemetry.csv"/>
 		<Property Name="TelemetryFilePath" Value="$(TraceDir)/$(TelemetryFileName)"/>
-		<Property Name="OutputPath" Value="$(AutomationOutputDir)\$(OutputName)"/>
-		<Property Name="OutputPath" Value="$(AutomationOutputDir)/$(OutputName)" If="'$(HostPlatform)' != 'Win64'"/>
-		<Property Name="ReportParams" Value="-ReportExportPath=&quot;$(OutputPath)&quot;"/>
+		<Property Name="ReportExportPath" Value="$(AutomationOutputDir)/$(OutputName)"/>
+		<Property Name="ReportParams" Value="-ReportExportPath=&quot;$(ReportExportPath)&quot;"/>
 		<Do If="$(IsBuildMachine)">
 			<Property Name="ReportExportURL" Value="$(AutomationReportBaseUrl)/$(OutputName)/index.html"/>
 			<Property Name="ReportParams" Value="$(ReportParams) -ReportURL=&quot;$(ReportExportURL)&quot; -WriteTestResultsForHorde -HordeTestDataKey=&quot;$(NodeName) $(Phase)$(PreflightSuffix)&quot;"/>
 		</Do>
-		<Property Name="LogOutputParams" Value="-logdir=&quot;$(OutputPath)&quot;"/> <!-- handled with care, using backslashes, so it's a valid path in the Horde web pages -->
+		<Property Name="LogOutputParams" Value="-logdir=&quot;$(AutomationOutputDir)/$(OutputName)&quot;"/>
 		<Property Name="AllGauntletArgs" Value="-project=$(Project) -platform=$(Platform) $(GauntletArgs) $(ReportParams) $(LogOutputParams) -artifactname=&quot;$(Phase)&quot;"/>
 
 		<!-- Run the Gauntlet Iteration Profile Tests -->
@@ -100,7 +99,7 @@
 		<Commandlet Name="SummarizeTrace" Arguments="-NoEnginePlugins -inputfile=&quot;$(TraceFilePath)&quot; -testname=&quot;$(TestName)&quot; $(SummarizeArgs)"/>
 
 		<!-- Save all files named for this run in the trace directory -->
-		<Copy From="$(TraceDir)/$(Project)$(Platform)$(TestName)*.*" To="$(OutputPath)"/>
+		<Copy From="$(TraceDir)/$(Project)$(Platform)$(TestName)*.*" To="$(AutomationOutputDir)"/>
 
 		<Do If="$(IsBuildMachine)">
 
@@ -116,11 +115,11 @@
 	<Property Name="RunGraphNodes" Value=""/>
 
 	<!-- Mimic of the Runner that retrieves the stashed telemetry file and copies it to the depot location -->
-	<Macro Name="BaselineResult" Arguments="OutputName;Platform;Project;Phase">
+	<Macro Name="BaselineResult" Arguments="Platform;Project;Phase">
 		<Property Name="TestName" Value="IterationProfile$(Phase)"/>
 		<Property Name="TelemetryFileName" Value="$(Project)$(Platform)$(TestName)Telemetry.csv"/>
 		<Property Name="LocalTelemetryFilePath" Value="$(BaselineDir)/$(TelemetryFileName)"/>
-		<Property Name="NetworkTelemetryFilePath" Value="$(BaselineAutomationNetworkOutputDir)/$(OutputName)/$(TelemetryFileName)"/>
+		<Property Name="NetworkTelemetryFilePath" Value="$(BaselineAutomationNetworkOutputDir)/$(TelemetryFileName)"/>
 		<Copy From="$(NetworkTelemetryFilePath)" To="$(LocalTelemetryFilePath)" Overwrite="true"/>
 		<Do If="!$(IsBuildMachine)">
 			<Spawn Exe="p4" Arguments="reconcile -ea $(LocalTelemetryFilePath)"/>
@@ -173,7 +172,7 @@
 				<!-- Baseline the Results -->
 				<Node Name="$(BaselineNodeName)" Requires="$(BaselineNodeRequires)">
 					<ForEach Name="Phase" Values="$(Phases)">
-						<Expand Name="BaselineResult" OutputName="$(OutputName)" Platform="$(Platform)" Project="$(Project)" Phase="$(Phase)"/>
+						<Expand Name="BaselineResult" Platform="$(Platform)" Project="$(Project)" Phase="$(Phase)"/>
 					</ForEach>
 				</Node>
 				<Property Name="BaselineGraphNodes" Value="$(BaselineGraphNodes);$(BaselineNodeName)"/>
