Add LLDB pretty printers to Linux installed builds

diff --git a/Engine/Build/InstalledEngineFilters.xml b/Engine/Build/InstalledEngineFilters.xml
--- a/Engine/Build/InstalledEngineFilters.xml
+++ b/Engine/Build/InstalledEngineFilters.xml
@@ -329,6 +329,7 @@
 		<!-- Build batch files -->
 		Engine/Build/BatchFiles/Linux/Build.sh
 		Engine/Build/BatchFiles/Linux/GenerateGDBInit.sh
+		Engine/Build/BatchFiles/Linux/GenerateLLDBInit.sh
 		Engine/Build/BatchFiles/Linux/FixDependencyFiles.sh
 		Engine/Build/BatchFiles/Linux/GenerateProjectFiles.sh
 		Engine/Build/BatchFiles/Linux/QASmokeManual.sh
