InstalledEngineFilters.xml: Add GetDotnetPath.bat

#trivial

diff --git a/Engine/Build/InstalledEngineFilters.xml b/Engine/Build/InstalledEngineFilters.xml
--- a/Engine/Build/InstalledEngineFilters.xml
+++ b/Engine/Build/InstalledEngineFilters.xml
@@ -244,6 +244,7 @@
 		Engine/Build/BatchFiles/Rebuild.bat
 		Engine/Build/BatchFiles/RunUAT.bat
 		Engine/Build/BatchFiles/GetMSBuildPath.bat
+		Engine/Build/BatchFiles/GetDotnetPath.bat
 
 		<!-- Build utilities -->
 		Engine/Build/Windows/cl-filter/cl-filter.exe
