fix bad merge for unrealgame targets android

diff --git a/Engine/Build/InstalledEngineBuild.xml b/Engine/Build/InstalledEngineBuild.xml
--- a/Engine/Build/InstalledEngineBuild.xml
+++ b/Engine/Build/InstalledEngineBuild.xml
@@ -1098,7 +1098,7 @@
 		<!-- Optional Target Platform requirements -->
 		<Property Name="InstalledRequirements" Value="$(InstalledRequirements);#UnrealGame Linux;#UnrealGame Linux Unstripped;#UnrealGame Linux Stripped" If="'$(WithLinux)' == true"/>
 		<Property Name="InstalledRequirements" Value="$(InstalledRequirements);#UnrealGame LinuxArm64;#UnrealGame LinuxArm64 Unstripped;#UnrealGame LinuxArm64 Stripped" If="'$(WithLinuxArm64)' == true"/>
-		<Property Name="InstalledRequirements" Value="$(InstalledRequirements);#UE4Game Android;#UE4Game Android Unstripped;#UE4Game Android Stripped" If="'$(WithAndroid)' == true"/>		
+		<Property Name="InstalledRequirements" Value="$(InstalledRequirements);#UnrealGame Android;#UnrealGame Android Unstripped;#UnrealGame Android Stripped" If="'$(WithAndroid)' == true"/>		
 		<Property Name="InstalledRequirements" Value="$(InstalledRequirements);#CompressedDDCLinux" If="'$(WithDDC)' == true And $(WithLinux)"/>
 
 		<!-- Feature packs -->
