fix other place bad merge happened with android

diff --git a/Engine/Build/InstalledEngineBuild.xml b/Engine/Build/InstalledEngineBuild.xml
--- a/Engine/Build/InstalledEngineBuild.xml
+++ b/Engine/Build/InstalledEngineBuild.xml
@@ -1139,8 +1139,8 @@
 				<Property Name="CopyInstalledExceptions" Value="$(CopyInstalledExceptions);$(CopyLinuxExceptions)"/>
 			</Do>
 			<Do If="'$(WithAndroid)' == true">
-				<Tag Files="#UE4Game Android" Except="#UE4Game Android Unstripped" With="#Installed Linux"/>
-				<Tag Files="#UE4Game Android Stripped" With="#Saved Output"/>
+				<Tag Files="#UnrealGame Android" Except="#UnrealGame Android Unstripped" With="#Installed Linux"/>
+				<Tag Files="#UnrealGame Android Stripped" With="#Saved Output"/>
 				<Property Name="Platforms" Value="$(Platforms)Android;"/>
 				<Property Name="CopyInstalledFilter" Value="$(CopyInstalledFilter);$(CopyAndroidFilter)"/>
 				<Property Name="CopyInstalledExceptions" Value="$(CopyInstalledExceptions);$(CopyAndroidExceptions)"/>
