IterationProfile Improvements

Call DDCCleanupCommandlet to tidy up the cache state so we don't get any cost/stat pollution during the profile
Set the UnusedFileAge to 0 to allow the DDCCleanupCommandlet to actually purge all ColdDerivedDataCache files
Set NODDCCLEANUP argument to prevent background cleanup of DDC data while we are profiling
Unset statnamedevents -- we don't track any of the legacy named events in telemetry, so let's not perpetuate it

diff --git a/Engine/Build/IterationProfile.xml b/Engine/Build/IterationProfile.xml
--- a/Engine/Build/IterationProfile.xml
+++ b/Engine/Build/IterationProfile.xml
@@ -54,11 +54,13 @@
 
 	<!-- Arguments for tracing -->
 	<Option Name="Phases" DefaultValue="Incr;Warm;Cold" Description="The list of phases to perform"/>
+	<Error Message="No phases specified" If="'$(Phases)' == ''"/>
+	<Property Name="TraceDir" Value="$(RootDir)/Engine/Programs/UnrealInsights/Saved/TraceSessions"/>
+
+	<!-- Arguments for phases, keep in mind that these are passed to both the DDCCleanupCommandlet and RunUnreal -->
 	<Option Name="ColdArgs" DefaultValue="-ddc=Cold" Description="Arguments to unreal for the cold phase"/>
 	<Option Name="IncrArgs" DefaultValue="-ddc=Warm" Description="Arguments to unreal for the incremental phase"/>
 	<Option Name="WarmArgs" DefaultValue="-ddc=Warm" Description="Arguments to unreal for the warm phase"/>
-	<Error Message="No phases specified" If="'$(Phases)' == ''"/>
-	<Property Name="TraceDir" Value="$(RootDir)/Engine/Programs/UnrealInsights/Saved/TraceSessions"/>
 
 	<!-- Gauntlet arguments for all runs -->
 	<Property Name="GauntletArgs" Value="-build=editor -UseEditor -verbose -report -tempdir=&quot;$(RootDir)/Tests&quot;"/>
@@ -70,7 +72,7 @@
 	<Property Name="ProfileArgs" Value="$(ProfileArgs) -EnablePlugins=$(ExtraPluginsCommaDelims)" If="'$(ExtraPluginsCommaDelims)' != ''"/>
 	<Property Name="ProfileArgs" Value="$(ProfileArgs) -TreatLogWarningsAsTestErrors=false -TreatLogErrorsAsTestErrors=false"/>
 	<Property Name="ProfileArgs" Value="$(ProfileArgs) -GameInstanceLostTimerSeconds=3600"/>
-	<Property Name="ProfileArgs" Value="$(ProfileArgs) -NoLoadLevelAtStartup -Benchmarking"/>
+	<Property Name="ProfileArgs" Value="$(ProfileArgs) -NoLoadLevelAtStartup -Benchmarking -NODDCCLEANUP"/>
 	<Property Name="ProfileArgs" Value="$(ProfileArgs) -NoShaderDistrib -VerboseShaderLogging"/>
 
 	<!-- Gauntlet arguments for summarize commandlet runs -->
@@ -92,6 +94,9 @@
 		<Property Name="LogOutputParams" Value="-logdir=&quot;$(AutomationOutputDir)/$(OutputName)&quot;"/>
 		<Property Name="AllGauntletArgs" Value="-project=$(Project) -platform=$(Platform) $(GauntletArgs) $(ReportParams) $(LogOutputParams) -artifactname=&quot;$(Phase)&quot;"/>
 
+		<!-- Do DDC cleanup outside the editor process, handles Cold cleanup without changing stats -->
+		<Commandlet Project="$(Project)" Name="DDCCleanup" Arguments="$(Args)"/>
+
 		<!-- Run the Gauntlet Iteration Profile Tests -->
 		<Command Name="RunUnreal" Arguments="$(AllGauntletArgs) $(ProfileArgs) $(Args) -tracefile=&quot;$(TraceFilePath)&quot;"/>
 
diff --git a/Engine/Config/BaseEngine.ini b/Engine/Config/BaseEngine.ini
--- a/Engine/Config/BaseEngine.ini
+++ b/Engine/Config/BaseEngine.ini
@@ -2078,7 +2078,7 @@ MinimumDaysToKeepFile=0
 Root=(Type=KeyLength, Length=120, Inner=AsyncPut)
 AsyncPut=(Type=AsyncPut, Inner=Hierarchy)
 Hierarchy=(Type=Hierarchical, Inner=Local)
-Local=(Type=FileSystem, ReadOnly=false, Clean=false, Flush=true, PurgeTransient=true, DeleteUnused=true, UnusedFileAge=1, FoldersToClean=-1, Path="%GAMEDIR%ColdDerivedDataCache\")
+Local=(Type=FileSystem, ReadOnly=false, Clean=false, Flush=true, PurgeTransient=true, DeleteUnused=true, UnusedFileAge=0, FoldersToClean=-1, Path="%GAMEDIR%ColdDerivedDataCache\")
 
 ; Inherits the same standalone cache from [Cold], for incremental performance measurement purposes
 [Warm]
diff --git a/Engine/Source/Programs/AutomationTool/Gauntlet/Unreal/Editor/EditorTests.cs b/Engine/Source/Programs/AutomationTool/Gauntlet/Unreal/Editor/EditorTests.cs
--- a/Engine/Source/Programs/AutomationTool/Gauntlet/Unreal/Editor/EditorTests.cs
+++ b/Engine/Source/Programs/AutomationTool/Gauntlet/Unreal/Editor/EditorTests.cs
@@ -97,7 +97,6 @@ namespace UnrealEditor
 				AppConfig.CommandLineParams.Add(string.Format("tracefile={0}", TraceFile));
 				AppConfig.CommandLineParams.Add("tracefiletrunc"); // replace existing
 				AppConfig.CommandLineParams.Add("trace=cpu,counters,bookmarks");
-				AppConfig.CommandLineParams.Add("statnamedevents");
 			}
 
 			if (TreatLogWarningsAsTestErrors.ToLower() == "false" || TreatLogWarningsAsTestErrors == "0")
