IterationProfile: Remove Test Report html generation, omit trace server for DCC commandlet

diff --git a/Engine/Build/IterationProfile.xml b/Engine/Build/IterationProfile.xml
--- a/Engine/Build/IterationProfile.xml
+++ b/Engine/Build/IterationProfile.xml
@@ -35,13 +35,11 @@
 	<!-- Automation uses its own network output directory -->
 	<Property Name="AutomationNetworkRootDir" Value="\\epicgames.net\root\Builds\Automation\Reports\IterationProfile"/>
 	<Property Name="AutomationNetworkRootDir" Value="/Volumes/Builds/Automation/Reports/IterationProfile" If="'$(HostPlatform)' != 'Win64'"/>
-	<Property Name="AutomationReportRootUrl" Value="http://automation.epicgames.net/reports/iterationprofile"/>
 	<Property Name="AutomationLocalOutputDir" Value="$(RootDir)/LocalBuilds/Reports"/>
 
 	<!-- Outputs for this run keeping slashes flipped the right way because it shows up on the webapp this way -->
 	<Property Name="AutomationNetworkOutputDir" Value="$(AutomationNetworkRootDir)\$(BuildName)"/>
 	<Property Name="AutomationNetworkOutputDir" Value="$(AutomationNetworkRootDir)/$(BuildName)" If="'$(HostPlatform)' != 'Win64'"/>
-	<Property Name="AutomationReportBaseUrl" Value="$(AutomationReportRootUrl)/$(BuildName)"/>
 	<Property Name="AutomationOutputDir" Value="$(AutomationLocalOutputDir)"/>
 	<Property Name="AutomationOutputDir" Value="$(AutomationNetworkOutputDir)" If="$(IsBuildMachine)"/>
 
@@ -63,7 +61,7 @@
 	<Option Name="WarmArgs" DefaultValue="-ddc=Warm" Description="Arguments to unreal for the warm phase"/>
 
 	<!-- Gauntlet arguments for all runs -->
-	<Property Name="GauntletArgs" Value="-build=editor -UseEditor -verbose -report -tempdir=&quot;$(RootDir)/Tests&quot;"/>
+	<Property Name="GauntletArgs" Value="-build=editor -UseEditor -verbose -tempdir=&quot;$(RootDir)/Tests&quot;"/>
 	<Property Name="GauntletArgs" Value="$(GauntletArgs) -branch=$(Branch) -ECBranch=&quot;$(Branch)&quot; -ECChangelist=&quot;$(Change)&quot;"/>
 	<Property Name="GauntletArgs" Value="$(GauntletArgs) -maxduration=28800"/>
 
@@ -75,9 +73,10 @@
 	<Property Name="ProfileArgs" Value="$(ProfileArgs) -NoLoadLevelAtStartup -Benchmarking -NODDCCLEANUP"/>
 	<Property Name="ProfileArgs" Value="$(ProfileArgs) -NoShaderDistrib -VerboseShaderLogging"/>
 
-	<!-- Gauntlet arguments for summarize commandlet runs -->
-	<Property Name="SummarizeArgs" Value="-notraceserver"/>
-	<Property Name="SummarizeArgs" Value="$(SummarizeArgs) -skipbaseline" If="$(RunAndBaseline)"/>
+	<!-- Gauntlet arguments for commandlet, summarize runs -->
+	<Property Name="CommandletArgs" Value="-notraceserver"/>
+	<Property Name="SummarizeArgs" Value=""/>
+	<Property Name="SummarizeArgs" Value="-skipbaseline" If="$(RunAndBaseline)"/>
 
 	<!-- The main test running macro logic that is expanded with platform and dependency node names -->
 	<Macro Name="RunTest" Arguments="OutputName;Platform;Project;NodeName;Phase;Args">
@@ -85,34 +84,27 @@
 		<Property Name="TraceFilePath" Value="$(TraceDir)/$(Project)$(Platform)$(TestName).utrace"/>
 		<Property Name="TelemetryFileName" Value="$(Project)$(Platform)$(TestName)Telemetry.csv"/>
 		<Property Name="TelemetryFilePath" Value="$(TraceDir)/$(TelemetryFileName)"/>
-		<Property Name="ReportExportPath" Value="$(AutomationOutputDir)/$(OutputName)"/>
-		<Property Name="ReportParams" Value="-ReportExportPath=&quot;$(ReportExportPath)&quot;"/>
-		<Do If="$(IsBuildMachine)">
-			<Property Name="ReportExportURL" Value="$(AutomationReportBaseUrl)/$(OutputName)/index.html"/>
-			<Property Name="ReportParams" Value="$(ReportParams) -ReportURL=&quot;$(ReportExportURL)&quot; -WriteTestResultsForHorde -HordeTestDataKey=&quot;$(NodeName) $(Phase)$(PreflightSuffix)&quot;"/>
-		</Do>
 		<Property Name="LogOutputParams" Value="-logdir=&quot;$(AutomationOutputDir)/$(OutputName)&quot;"/>
+		<Property Name="ReportParams" Value="-WriteTestResultsForHorde -HordeTestDataKey=&quot;$(NodeName) $(Phase)$(PreflightSuffix)&quot;"/>
 		<Property Name="AllGauntletArgs" Value="-project=$(Project) -platform=$(Platform) $(GauntletArgs) $(ReportParams) $(LogOutputParams) -artifactname=&quot;$(Phase)&quot;"/>
 
 		<!-- Do DDC cleanup outside the editor process, handles Cold cleanup without changing stats -->
-		<Commandlet Project="$(Project)" Name="DDCCleanup" Arguments="$(Args)"/>
+		<Commandlet Project="$(Project)" Name="DDCCleanup" Arguments="$(Args) $(CommandletArgs)"/>
 
 		<!-- Run the Gauntlet Iteration Profile Tests -->
 		<Command Name="RunUnreal" Arguments="$(AllGauntletArgs) $(ProfileArgs) $(Args) -tracefile=&quot;$(TraceFilePath)&quot;"/>
 
 		<!-- Translate the utrace file into csv data -->
-		<Commandlet Project="$(Project)" Name="SummarizeTrace" Arguments="-inputfile=&quot;$(TraceFilePath)&quot; -testname=&quot;$(TestName)&quot; $(SummarizeArgs)"/>
+		<Commandlet Project="$(Project)" Name="SummarizeTrace" Arguments="-inputfile=&quot;$(TraceFilePath)&quot; -testname=&quot;$(TestName)&quot; $(SummarizeArgs) $(CommandletArgs)"/>
 
 		<!-- Save all files named for this run in the trace directory -->
 		<Copy From="$(TraceDir)/$(Project)$(Platform)$(TestName)*.*" To="$(AutomationOutputDir)"/>
 
 		<Do If="$(IsBuildMachine)">
-
 			<Do If="!$(IsPreflight)">
 				<!-- Publish telemetry to the db -->
 				<Command Name="PublishUnrealAutomationTelemetry" Arguments="-CSVFile=&quot;$(TelemetryFilePath)&quot; -project=$(Project) -changelist=$(Change) -branch=$(Branch)"/>
 			</Do>
-
 		</Do>
 	</Macro>
 
