IterationProfile: Run multiple iterations

Each iteration's telemetry will be averaged by tableau.
No preflight as this build graph xml file doesn't impact any other builds.

diff --git a/Engine/Build/IterationProfile.xml b/Engine/Build/IterationProfile.xml
--- a/Engine/Build/IterationProfile.xml
+++ b/Engine/Build/IterationProfile.xml
@@ -60,6 +60,9 @@
 	<Option Name="IncrArgs" DefaultValue="-ddc=Warm" Description="Arguments to unreal for the incremental phase"/>
 	<Option Name="WarmArgs" DefaultValue="-ddc=Warm" Description="Arguments to unreal for the warm phase"/>
 
+	<!-- Iterations to run and average together -->
+	<Property Name="Iterations" Value="1;2;3"/>
+
 	<!-- Gauntlet arguments for all runs -->
 	<Property Name="GauntletArgs" Value="-build=editor -UseEditor -verbose -tempdir=&quot;$(RootDir)/Tests&quot;"/>
 	<Property Name="GauntletArgs" Value="$(GauntletArgs) -branch=$(Branch) -ECBranch=&quot;$(Branch)&quot; -ECChangelist=&quot;$(Change)&quot;"/>
@@ -81,31 +84,40 @@
 	<!-- The main test running macro logic that is expanded with platform and dependency node names -->
 	<Macro Name="RunTest" Arguments="OutputName;Platform;Project;NodeName;Phase;Args">
 		<Property Name="TestName" Value="IterationProfile$(Phase)"/>
-		<Property Name="TraceFilePath" Value="$(TraceDir)/$(Project)$(Platform)$(TestName).utrace"/>
-		<Property Name="TelemetryFileName" Value="$(Project)$(Platform)$(TestName)Telemetry.csv"/>
-		<Property Name="TelemetryFilePath" Value="$(TraceDir)/$(TelemetryFileName)"/>
-		<Property Name="LogOutputParams" Value="-logdir=&quot;$(AutomationOutputDir)/$(OutputName)&quot;"/>
-		<Property Name="ReportParams" Value="-WriteTestResultsForHorde -HordeTestDataKey=&quot;$(NodeName) $(Phase)$(PreflightSuffix)&quot;"/>
-		<Property Name="AllGauntletArgs" Value="-project=$(Project) -platform=$(Platform) $(GauntletArgs) $(ReportParams) $(LogOutputParams) -artifactname=&quot;$(Phase)&quot;"/>
 
-		<!-- Do DDC cleanup outside the editor process, handles Cold cleanup without changing stats -->
-		<Commandlet Project="$(Project)" Name="DDCCleanup" Arguments="$(Args) $(CommandletArgs)"/>
+		<ForEach Name="Iteration" Values="$(Iterations)">
+
+			<!-- Only ever do one iteration of incremental, because it will warm up the workspace -->
+			<Do If="'$(Phase)' != 'Incr' or '$(Iteration)' == '1'">
+				<Property Name="IterationName" Value="$(Project)$(Platform)$(TestName)$(Iteration)"/>
+				<Property Name="TraceFilePath" Value="$(TraceDir)/$(IterationName).utrace"/>
+				<Property Name="TelemetryFileName" Value="$(IterationName)Telemetry.csv"/>
+				<Property Name="TelemetryFilePath" Value="$(TraceDir)/$(TelemetryFileName)"/>
+				<Property Name="LogOutputParams" Value="-logdir=&quot;$(AutomationOutputDir)/$(OutputName)&quot;"/>
+				<Property Name="ReportParams" Value="-WriteTestResultsForHorde -HordeTestDataKey=&quot;$(NodeName) $(Phase)$(PreflightSuffix)&quot;"/>
+				<Property Name="AllGauntletArgs" Value="-project=$(Project) -platform=$(Platform) $(GauntletArgs) $(ReportParams) $(LogOutputParams) -artifactname=&quot;$(Phase)&quot;"/>
+
+				<!-- Do DDC cleanup outside the editor process, handles Cold cleanup without changing stats -->
+				<Commandlet Project="$(Project)" Name="DDCCleanup" Arguments="$(Args) $(CommandletArgs)"/>
 
-		<!-- Run the Gauntlet Iteration Profile Tests -->
-		<Command Name="RunUnreal" Arguments="$(AllGauntletArgs) $(ProfileArgs) $(Args) -tracefile=&quot;$(TraceFilePath)&quot;"/>
+				<!-- Run the Gauntlet Iteration Profile Tests -->
+				<Command Name="RunUnreal" Arguments="$(AllGauntletArgs) $(ProfileArgs) $(Args) -tracefile=&quot;$(TraceFilePath)&quot;"/>
 
-		<!-- Translate the utrace file into csv data -->
-		<Commandlet Project="$(Project)" Name="SummarizeTrace" Arguments="-inputfile=&quot;$(TraceFilePath)&quot; -testname=&quot;$(TestName)&quot; $(SummarizeArgs) $(CommandletArgs)"/>
+				<!-- Translate the utrace file into csv data -->
+				<Commandlet Project="$(Project)" Name="SummarizeTrace" Arguments="-inputfile=&quot;$(TraceFilePath)&quot; -testname=&quot;$(TestName)&quot; $(SummarizeArgs) $(CommandletArgs)"/>
 
-		<!-- Save all files named for this run in the trace directory -->
-		<Copy From="$(TraceDir)/$(Project)$(Platform)$(TestName)*.*" To="$(AutomationOutputDir)"/>
+				<!-- Save all files named for this run in the trace directory -->
+				<Copy From="$(TraceDir)/$(IterationName)*.*" To="$(AutomationOutputDir)"/>
 
-		<Do If="$(IsBuildMachine)">
-			<Do If="!$(IsPreflight)">
-				<!-- Publish telemetry to the db -->
-				<Command Name="PublishUnrealAutomationTelemetry" Arguments="-CSVFile=&quot;$(TelemetryFilePath)&quot; -project=$(Project) -changelist=$(Change) -branch=$(Branch)"/>
+				<Do If="$(IsBuildMachine)">
+					<Do If="!$(IsPreflight)">
+						<!-- Publish telemetry to the db -->
+						<Command Name="PublishUnrealAutomationTelemetry" Arguments="-CSVFile=&quot;$(TelemetryFilePath)&quot; -project=$(Project) -changelist=$(Change) -branch=$(Branch)"/>
+					</Do>
+				</Do>
 			</Do>
-		</Do>
+
+		</ForEach> <!-- Iterations -->
 	</Macro>
 
 	<!-- This will be set to the names of all nodes we want to execute based on specified platforms/options -->
@@ -114,7 +126,8 @@
 	<!-- Mimic of the Runner that retrieves the stashed telemetry file and copies it to the depot location -->
 	<Macro Name="BaselineResult" Arguments="Platform;Project;Phase">
 		<Property Name="TestName" Value="IterationProfile$(Phase)"/>
-		<Property Name="TelemetryFileName" Value="$(Project)$(Platform)$(TestName)Telemetry.csv"/>
+		<Property Name="IterationName" Value="$(Project)$(Platform)$(TestName)1"/> <!-- Note the 1! We use Iteration 1 for baselines for now -->
+		<Property Name="TelemetryFileName" Value="$(IterationName)Telemetry.csv"/>
 		<Property Name="LocalTelemetryFilePath" Value="$(BaselineDir)/$(TelemetryFileName)"/>
 		<Property Name="NetworkTelemetryFilePath" Value="$(BaselineAutomationNetworkOutputDir)/$(TelemetryFileName)"/>
 		<Copy From="$(NetworkTelemetryFilePath)" To="$(LocalTelemetryFilePath)" Overwrite="true"/>
