Rechristen IterationProfile -> EditorPerf for clarity

Part 1 - branch files to new names
No preflight as this just creates some new files, which shouldn't cause any issues

#rnx

diff --git a/Engine/Build/BatchFiles/EditorPerf.bat b/Engine/Build/BatchFiles/EditorPerf.bat
new file mode 100755
--- /dev/null
+++ b/Engine/Build/BatchFiles/EditorPerf.bat
@@ -0,0 +1 @@
+@call %~dp0RunUAT.bat BuildGraph -Script=Engine/Build/IterationProfile.xml -Target="Run Iteration Profile Tests" -set:Platforms=Win64 -append:Projects="AncientGame" -append:Projects="EngineTest" -append:Projects="ShooterGame" %*
\ No newline at end of file
diff --git a/Engine/Build/BatchFiles/EditorPerfBaseline.bat b/Engine/Build/BatchFiles/EditorPerfBaseline.bat
new file mode 100755
--- /dev/null
+++ b/Engine/Build/BatchFiles/EditorPerfBaseline.bat
@@ -0,0 +1 @@
+@call %~dp0RunUAT.bat BuildGraph -Script=Engine/Build/IterationProfile.xml -Target="Baseline Iteration Profile Results" -set:Platforms=Win64 -set:Phases="Cold;Warm" -append:Projects="AncientGame" -append:Projects="EngineTest" -append:Projects="ShooterGame" %*
\ No newline at end of file
diff --git a/Engine/Build/EditorPerf.xml b/Engine/Build/EditorPerf.xml
new file mode 100644
--- /dev/null
+++ b/Engine/Build/EditorPerf.xml
@@ -0,0 +1,212 @@
+<?xml version='1.0' ?>
+<BuildGraph xmlns="http://www.epicgames.com/BuildGraph" xmlns:xsi="http://www.w3.org/2001/XMLSchema-instance" xsi:schemaLocation="http://www.epicgames.com/BuildGraph ../../../Build/Graph/Schema.xsd" >
+
+	<!-- Horde Options -->	
+	<EnvVar Name="UE_HORDE_JOBID"/>
+	<Option Name="Horde" Restrict="true|false" DefaultValue="false" Description="Whether the build is running under Horde"/>
+	<Option Name="PreflightChange" DefaultValue="" Description="Preflight CL number if preflight, empty otherwise"/>
+	<Option Name="SetPreflightFlag" DefaultValue="true" Description="Whether to set the IsPreflight flag to true for preflights. Use with caution: this will cause things like build versions to be set to their non-preflight form."/>
+	<Property Name="IsPreflight" Value="false"/>
+	<Property Name="IsPreflight" Value="true" If="'$(PreflightChange)' != '' And $(SetPreflightFlag)"/>
+	<Property Name="PreflightSuffix" Value=""/>
+	<Property Name="PreflightSuffix" Value="-PF-$(PreflightChange)-$(UE_HORDE_JOBID)" If="$(IsPreflight)"/>
+	<Property Name="BuildName" Value="$(EscapedBranch)-CL-$(Change)$(PreflightSuffix)"/>
+
+	<!-- Projects setup-->
+	<Option Name="Projects" DefaultValue="false" Restrict=".+" Description="The projects to test"/>
+	<Error Message="No project names specified" If="'$(Projects)' == 'false'"/>
+
+	<!-- Platform Options -->
+	<Option Name="Platforms" DefaultValue="Win64" Description="The list of platforms to profile"/>
+	<Property Name="Win64CompileAgentType" Value="Win64"/>
+	<Property Name="MacCompileAgentType" Value="TestMac"/>
+	<Error Message="No platforms specified" If="'$(Platforms)' == ''"/>
+
+	<!-- Build Options -->
+	<Property Name="ExtraPluginsCommaDelims" Value="EditorTests"/>
+	<Property Name="ExtraPluginsPlusDelims" Value=""/> <!-- populated programmatically -->
+	<ForEach Name="Plugin" Values="$(ExtraPluginsCommaDelims)">
+		<Property Name="ExtraPluginString" Value="$(Plugin)" If="'$(ExtraPluginsPlusDelims)' == ''"/>
+		<Property Name="ExtraPluginString" Value="$(ExtraPluginsPlusDelims)+$(Plugin)" If="'$(ExtraPluginsPlusDelims)' != ''"/>
+		<Property Name="ExtraPluginsPlusDelims" Value="$(ExtraPluginString)"/>
+	</ForEach>
+	<Property Name="EditorCompileArgs" Value="-BuildPlugin=$(ExtraPluginsPlusDelims)"/>
+
+	<!-- Automation uses its own network output directory -->
+	<Property Name="AutomationNetworkRootDir" Value="\\epicgames.net\root\Builds\Automation\Reports\IterationProfile"/>
+	<Property Name="AutomationNetworkRootDir" Value="/Volumes/Builds/Automation/Reports/IterationProfile" If="'$(HostPlatform)' != 'Win64'"/>
+	<Property Name="AutomationLocalOutputDir" Value="$(RootDir)/LocalBuilds/Reports"/>
+
+	<!-- Outputs for this run keeping slashes flipped the right way because it shows up on the webapp this way -->
+	<Property Name="AutomationNetworkOutputDir" Value="$(AutomationNetworkRootDir)\$(BuildName)"/>
+	<Property Name="AutomationNetworkOutputDir" Value="$(AutomationNetworkRootDir)/$(BuildName)" If="'$(HostPlatform)' != 'Win64'"/>
+	<Property Name="AutomationOutputDir" Value="$(AutomationLocalOutputDir)"/>
+	<Property Name="AutomationOutputDir" Value="$(AutomationNetworkOutputDir)" If="$(IsBuildMachine)"/>
+
+	<!-- Options for baselining -->
+	<Option Name="RunAndBaseline" DefaultValue="false" Description="Enable graph dependencies for running and setting the baseline in a single build."/>
+	<Option Name="BaselineDir" DefaultValue="$(RootDir)/Engine/Build/Baseline" Description="The revision controlled location of the baseline telemetry files."/>
+	<Option Name="BaselineBuildName" DefaultValue="$(BuildName)" Description="The BuildName of the build to set as the new baseline to compare against."/>
+	<Property Name="BaselineAutomationNetworkOutputDir" Value="$(AutomationNetworkRootDir)/$(BaselineBuildName)"/>
+	<Property Name="BaselineAutomationNetworkOutputDir" Value="$(AutomationOutputDir)" If="$(RunAndBaseline)"/>
+
+	<!-- Arguments for tracing -->
+	<Option Name="Phases" DefaultValue="Incr;Warm;Cold" Description="The list of phases to perform"/>
+	<Error Message="No phases specified" If="'$(Phases)' == ''"/>
+	<Property Name="TraceDir" Value="$(RootDir)/Engine/Programs/UnrealInsights/Saved/TraceSessions"/>
+
+	<!-- Arguments for phases, keep in mind that these are passed to both the DDCCleanupCommandlet and RunUnreal -->
+	<Option Name="ColdArgs" DefaultValue="-ddc=Cold" Description="Arguments to unreal for the cold phase"/>
+	<Option Name="IncrArgs" DefaultValue="-ddc=Warm" Description="Arguments to unreal for the incremental phase"/>
+	<Option Name="WarmArgs" DefaultValue="-ddc=Warm" Description="Arguments to unreal for the warm phase"/>
+
+	<!-- Iterations to run and average together -->
+	<Property Name="Iterations" Value="1;2;3"/>
+
+	<!-- Gauntlet arguments for all runs -->
+	<Property Name="GauntletArgs" Value="-build=editor -UseEditor -verbose -tempdir=&quot;$(RootDir)/Tests&quot;"/>
+	<Property Name="GauntletArgs" Value="$(GauntletArgs) -branch=$(Branch) -ECBranch=&quot;$(Branch)&quot; -ECChangelist=&quot;$(Change)&quot;"/>
+	<Property Name="GauntletArgs" Value="$(GauntletArgs) -maxduration=28800"/>
+
+	<!-- Gauntlet arguments for profile runs -->
+	<Property Name="ProfileArgs" Value="-test=EditorTests -testfilter=Group:Iteration -tracefiletrunc"/>
+	<Property Name="ProfileArgs" Value="$(ProfileArgs) -EnablePlugins=$(ExtraPluginsCommaDelims)" If="'$(ExtraPluginsCommaDelims)' != ''"/>
+	<Property Name="ProfileArgs" Value="$(ProfileArgs) -TreatLogWarningsAsTestErrors=false -TreatLogErrorsAsTestErrors=false"/>
+	<Property Name="ProfileArgs" Value="$(ProfileArgs) -GameInstanceLostTimerSeconds=3600"/>
+	<Property Name="ProfileArgs" Value="$(ProfileArgs) -NoLoadLevelAtStartup -Benchmarking -NODDCCLEANUP"/>
+	<Property Name="ProfileArgs" Value="$(ProfileArgs) -NoShaderDistrib -VerboseShaderLogging"/>
+
+	<!-- Gauntlet arguments for commandlet, summarize runs -->
+	<Property Name="CommandletArgs" Value="-notraceserver"/>
+	<Property Name="SummarizeArgs" Value=""/>
+	<Property Name="SummarizeArgs" Value="-skipbaseline" If="$(RunAndBaseline)"/>
+
+	<!-- The main test running macro logic that is expanded with platform and dependency node names -->
+	<Macro Name="RunTest" Arguments="OutputName;Platform;Project;NodeName;Phase;Args">
+		<Property Name="TestName" Value="IterationProfile$(Phase)"/>
+
+		<ForEach Name="Iteration" Values="$(Iterations)">
+
+			<!-- Only ever do one iteration of incremental, because it will warm up the workspace -->
+			<Do If="'$(Phase)' != 'Incr' or '$(Iteration)' == '1'">
+				<Property Name="IterationName" Value="$(Project)$(Platform)$(TestName)$(Iteration)"/>
+				<Property Name="TraceFilePath" Value="$(TraceDir)/$(IterationName).utrace"/>
+				<Property Name="TelemetryFileName" Value="$(IterationName)Telemetry.csv"/>
+				<Property Name="TelemetryFilePath" Value="$(TraceDir)/$(TelemetryFileName)"/>
+				<Property Name="LogOutputParams" Value="-logdir=&quot;$(AutomationOutputDir)/$(OutputName)&quot;"/>
+				<Property Name="ReportParams" Value="-WriteTestResultsForHorde -HordeTestDataKey=&quot;$(NodeName) $(Phase)$(PreflightSuffix)&quot;"/>
+				<Property Name="AllGauntletArgs" Value="-project=$(Project) -platform=$(Platform) $(GauntletArgs) $(ReportParams) $(LogOutputParams) -artifactname=&quot;$(Phase)&quot;"/>
+
+				<!-- Do DDC cleanup outside the editor process, handles Cold cleanup without changing stats -->
+				<Commandlet Project="$(Project)" Name="DDCCleanup" Arguments="$(Args) $(CommandletArgs)"/>
+
+				<!-- Run the Gauntlet Iteration Profile Tests -->
+				<Command Name="RunUnreal" Arguments="$(AllGauntletArgs) $(ProfileArgs) $(Args) -tracefile=&quot;$(TraceFilePath)&quot;"/>
+
+				<!-- Translate the utrace file into csv data -->
+				<Commandlet Project="$(Project)" Name="SummarizeTrace" Arguments="-inputfile=&quot;$(TraceFilePath)&quot; -testname=&quot;$(TestName)&quot; $(SummarizeArgs) $(CommandletArgs)"/>
+
+				<!-- Save all files named for this run in the trace directory -->
+				<Copy From="$(TraceDir)/$(IterationName)*.*" To="$(AutomationOutputDir)"/>
+
+				<Do If="$(IsBuildMachine)">
+					<Do If="!$(IsPreflight)">
+						<!-- Publish telemetry to the db -->
+						<Command Name="PublishUnrealAutomationTelemetry" Arguments="-CSVFile=&quot;$(TelemetryFilePath)&quot; -project=$(Project) -changelist=$(Change) -branch=$(Branch)"/>
+					</Do>
+				</Do>
+			</Do>
+
+		</ForEach> <!-- Iterations -->
+	</Macro>
+
+	<!-- This will be set to the names of all nodes we want to execute based on specified platforms/options -->
+	<Property Name="RunGraphNodes" Value=""/>
+
+	<!-- Mimic of the Runner that retrieves the stashed telemetry file and copies it to the depot location -->
+	<Macro Name="BaselineResult" Arguments="Platform;Project;Phase">
+		<Property Name="TestName" Value="IterationProfile$(Phase)"/>
+		<Property Name="IterationName" Value="$(Project)$(Platform)$(TestName)1"/> <!-- Note the 1! We use Iteration 1 for baselines for now -->
+		<Property Name="TelemetryFileName" Value="$(IterationName)Telemetry.csv"/>
+		<Property Name="LocalTelemetryFilePath" Value="$(BaselineDir)/$(TelemetryFileName)"/>
+		<Property Name="NetworkTelemetryFilePath" Value="$(BaselineAutomationNetworkOutputDir)/$(TelemetryFileName)"/>
+		<Copy From="$(NetworkTelemetryFilePath)" To="$(LocalTelemetryFilePath)" Overwrite="true"/>
+		<Do If="!$(IsBuildMachine)">
+			<Spawn Exe="p4" Arguments="reconcile -ea $(LocalTelemetryFilePath)"/>
+		</Do>
+	</Macro>
+
+	<!-- This will be set to the names of all nodes we want to execute based on specified platforms/options -->
+	<Property Name="BaselineGraphNodes" Value=""/>
+
+	<!-- ForEach Platform to run, generate graph nodes -->
+	<ForEach Name="Platform" Values="$(Platforms)">
+
+		<!-- Project-independent prereq compiles -->
+		<Agent Name="Prerequisites Agent $(Platform)" Type="$($(Platform)CompileAgentType)">
+			<Node Name="Compile Prerequisites $(Platform)">
+				<Compile Target="UnrealHeaderTool" Platform="$(Platform)" Configuration="Development"/>
+				<Compile Target="ShaderCompileWorker" Platform="$(Platform)" Configuration="Development"/>
+				<Compile Target="CrashReportClient" Platform="$(Platform)" Configuration="Shipping"/>
+				<Compile Target="CrashReportClientEditor" Platform="$(Platform)" Configuration="Shipping"/>
+			</Node>
+
+			<!-- ForEach Project to run, generate the compile graph nodes -->
+			<ForEach Name="Project" Values="$(Projects)">
+				<Node Name="Compile $(Project) Editor $(Platform)" Requires="Compile Prerequisites $(Platform)">
+					<Compile Target="$(Project)Editor" Platform="$(Platform)" Configuration="Development" Arguments="$(EditorCompileArgs)"/>
+				</Node>
+			</ForEach>
+		</Agent>
+
+		<!-- ForEach Project to run, generate the graph nodes -->
+		<ForEach Name="Project" Values="$(Projects)">
+
+			<!-- Project execution -->
+			<Agent Name="Automated Test Agent $(Project) $(Platform)" Type="IterationProfile$(Platform)$(Project)">
+				<Property Name="OutputName" Value="$(Project)$(Platform)IterationProfile"/>
+				<Property Name="RunNodeName" Value="Run $(OutputName)"/>
+				<Property Name="RunNodeRequires" Value="Compile Prerequisites $(Platform);Compile $(Project) Editor $(Platform)"/>
+				<Property Name="BaselineNodeName" Value="Baseline $(OutputName)"/>
+				<Property Name="BaselineNodeRequires" Value=""/>
+				<Property Name="BaselineNodeRequires" Value="$(RunNodeName)" If="$(RunAndBaseline)"/>
+
+				<!-- Run the Tests -->
+				<Node Name="$(RunNodeName)" Requires="$(RunNodeRequires)">
+					<ForEach Name="Phase" Values="$(Phases)">
+						<Expand Name="RunTest" OutputName="$(OutputName)" Platform="$(Platform)" Project="$(Project)" NodeName="$(RunNodeName)" Phase="$(Phase)" Args="$($(Phase)Args)"/>
+					</ForEach>
+				</Node>
+				<Property Name="RunGraphNodes" Value="$(RunGraphNodes);$(RunNodeName)"/>
+
+				<!-- Baseline the Results -->
+				<Node Name="$(BaselineNodeName)" Requires="$(BaselineNodeRequires)">
+					<ForEach Name="Phase" Values="$(Phases)">
+						<Expand Name="BaselineResult" Platform="$(Platform)" Project="$(Project)" Phase="$(Phase)"/>
+					</ForEach>
+				</Node>
+				<Property Name="BaselineGraphNodes" Value="$(BaselineGraphNodes);$(BaselineNodeName)"/>
+
+				<Label Category="Test" Name="$(Project) Iteration Profile $(Platform)" Requires="$(RunNodeName)"/>
+			</Agent>
+
+		</ForEach> <!-- Projects -->
+
+	</ForEach> <!-- Platforms -->
+
+	<!-- These Nodes run the Editor and gather profile data -->
+	<Aggregate Name="Run Iteration Profile Tests" Requires="$(RunGraphNodes)"/>
+
+	<!-- These Nodes copies the CSV telemetry files to special location in the depot, thus making them the "baselines" -->
+	<Aggregate Name="Baseline Iteration Profile Results" Requires="$(BaselineGraphNodes)"/>
+
+	<!-- This Node submits the baseline CSV files to Perforce -->
+	<!--
+	<Agent Name="Submission Agent Win64" Type="Win64">
+		<Node Name="Submit Baseline To Perforce" Requires="$(BaselineGraphNodes)">
+			<Submit Description="[CL $(Change)] Updated IterationProfile baselines" Files="$(BaselineFiles)"/>
+		</Node>
+	</Agent>
+	-->
+
+</BuildGraph>
