Rechristen IterationProfile -> EditorPerf for clarity

Part 2 - modify branched files, delete old files

#rnx

diff --git a/Engine/Build/BatchFiles/EditorPerf.bat b/Engine/Build/BatchFiles/EditorPerf.bat
--- a/Engine/Build/BatchFiles/EditorPerf.bat
+++ b/Engine/Build/BatchFiles/EditorPerf.bat
@@ -1 +1 @@
-@call %~dp0RunUAT.bat BuildGraph -Script=Engine/Build/IterationProfile.xml -Target="Run Iteration Profile Tests" -set:Platforms=Win64 -append:Projects="AncientGame" -append:Projects="EngineTest" -append:Projects="ShooterGame" %*
\ No newline at end of file
+@call %~dp0RunUAT.bat BuildGraph -Script=Engine/Build/EditorPerf.xml -Target="Run EditorPerf Tests" -set:Platforms=Win64 -append:Projects="AncientGame" -append:Projects="EngineTest" -append:Projects="ShooterGame" %*
\ No newline at end of file
diff --git a/Engine/Build/BatchFiles/EditorPerfBaseline.bat b/Engine/Build/BatchFiles/EditorPerfBaseline.bat
--- a/Engine/Build/BatchFiles/EditorPerfBaseline.bat
+++ b/Engine/Build/BatchFiles/EditorPerfBaseline.bat
@@ -1 +1 @@
-@call %~dp0RunUAT.bat BuildGraph -Script=Engine/Build/IterationProfile.xml -Target="Baseline Iteration Profile Results" -set:Platforms=Win64 -set:Phases="Cold;Warm" -append:Projects="AncientGame" -append:Projects="EngineTest" -append:Projects="ShooterGame" %*
\ No newline at end of file
+@call %~dp0RunUAT.bat BuildGraph -Script=Engine/Build/EditorPerf.xml -Target="Baseline EditorPerf Results" -set:Platforms=Win64 -set:Phases="Cold;Warm" -append:Projects="AncientGame" -append:Projects="EngineTest" -append:Projects="ShooterGame" %*
\ No newline at end of file
diff --git a/Engine/Build/BatchFiles/IterationProfile.bat b/Engine/Build/BatchFiles/IterationProfile.bat
deleted file mode 100755
diff --git a/Engine/Build/BatchFiles/IterationProfileBaseline.bat b/Engine/Build/BatchFiles/IterationProfileBaseline.bat
deleted file mode 100755
diff --git a/Engine/Build/EditorPerf.xml b/Engine/Build/EditorPerf.xml
--- a/Engine/Build/EditorPerf.xml
+++ b/Engine/Build/EditorPerf.xml
@@ -33,8 +33,8 @@
 	<Property Name="EditorCompileArgs" Value="-BuildPlugin=$(ExtraPluginsPlusDelims)"/>
 
 	<!-- Automation uses its own network output directory -->
-	<Property Name="AutomationNetworkRootDir" Value="\\epicgames.net\root\Builds\Automation\Reports\IterationProfile"/>
-	<Property Name="AutomationNetworkRootDir" Value="/Volumes/Builds/Automation/Reports/IterationProfile" If="'$(HostPlatform)' != 'Win64'"/>
+	<Property Name="AutomationNetworkRootDir" Value="\\epicgames.net\root\Builds\Automation\Reports\EditorPerf"/>
+	<Property Name="AutomationNetworkRootDir" Value="/Volumes/Builds/Automation/Reports/EditorPerf" If="'$(HostPlatform)' != 'Win64'"/>
 	<Property Name="AutomationLocalOutputDir" Value="$(RootDir)/LocalBuilds/Reports"/>
 
 	<!-- Outputs for this run keeping slashes flipped the right way because it shows up on the webapp this way -->
@@ -83,7 +83,7 @@
 
 	<!-- The main test running macro logic that is expanded with platform and dependency node names -->
 	<Macro Name="RunTest" Arguments="OutputName;Platform;Project;NodeName;Phase;Args">
-		<Property Name="TestName" Value="IterationProfile$(Phase)"/>
+		<Property Name="TestName" Value="EditorPerf$(Phase)"/>
 
 		<ForEach Name="Iteration" Values="$(Iterations)">
 
@@ -100,7 +100,7 @@
 				<!-- Do DDC cleanup outside the editor process, handles Cold cleanup without changing stats -->
 				<Commandlet Project="$(Project)" Name="DDCCleanup" Arguments="$(Args) $(CommandletArgs)"/>
 
-				<!-- Run the Gauntlet Iteration Profile Tests -->
+				<!-- Run the Gauntlet EditorPerf Tests -->
 				<Command Name="RunUnreal" Arguments="$(AllGauntletArgs) $(ProfileArgs) $(Args) -tracefile=&quot;$(TraceFilePath)&quot;"/>
 
 				<!-- Translate the utrace file into csv data -->
@@ -125,7 +125,7 @@
 
 	<!-- Mimic of the Runner that retrieves the stashed telemetry file and copies it to the depot location -->
 	<Macro Name="BaselineResult" Arguments="Platform;Project;Phase">
-		<Property Name="TestName" Value="IterationProfile$(Phase)"/>
+		<Property Name="TestName" Value="EditorPerf$(Phase)"/>
 		<Property Name="IterationName" Value="$(Project)$(Platform)$(TestName)1"/> <!-- Note the 1! We use Iteration 1 for baselines for now -->
 		<Property Name="TelemetryFileName" Value="$(IterationName)Telemetry.csv"/>
 		<Property Name="LocalTelemetryFilePath" Value="$(BaselineDir)/$(TelemetryFileName)"/>
@@ -163,8 +163,8 @@
 		<ForEach Name="Project" Values="$(Projects)">
 
 			<!-- Project execution -->
-			<Agent Name="Automated Test Agent $(Project) $(Platform)" Type="IterationProfile$(Platform)$(Project)">
-				<Property Name="OutputName" Value="$(Project)$(Platform)IterationProfile"/>
+			<Agent Name="Automated Test Agent $(Project) $(Platform)" Type="EditorPerf$(Platform)$(Project)">
+				<Property Name="OutputName" Value="$(Project)$(Platform)EditorPerf"/>
 				<Property Name="RunNodeName" Value="Run $(OutputName)"/>
 				<Property Name="RunNodeRequires" Value="Compile Prerequisites $(Platform);Compile $(Project) Editor $(Platform)"/>
 				<Property Name="BaselineNodeName" Value="Baseline $(OutputName)"/>
@@ -187,7 +187,7 @@
 				</Node>
 				<Property Name="BaselineGraphNodes" Value="$(BaselineGraphNodes);$(BaselineNodeName)"/>
 
-				<Label Category="Test" Name="$(Project) Iteration Profile $(Platform)" Requires="$(RunNodeName)"/>
+				<Label Category="Test" Name="$(Project) EditorPerf $(Platform)" Requires="$(RunNodeName)"/>
 			</Agent>
 
 		</ForEach> <!-- Projects -->
@@ -195,16 +195,16 @@
 	</ForEach> <!-- Platforms -->
 
 	<!-- These Nodes run the Editor and gather profile data -->
-	<Aggregate Name="Run Iteration Profile Tests" Requires="$(RunGraphNodes)"/>
+	<Aggregate Name="Run EditorPerf Tests" Requires="$(RunGraphNodes)"/>
 
 	<!-- These Nodes copies the CSV telemetry files to special location in the depot, thus making them the "baselines" -->
-	<Aggregate Name="Baseline Iteration Profile Results" Requires="$(BaselineGraphNodes)"/>
+	<Aggregate Name="Baseline EditorPerf Results" Requires="$(BaselineGraphNodes)"/>
 
 	<!-- This Node submits the baseline CSV files to Perforce -->
 	<!--
 	<Agent Name="Submission Agent Win64" Type="Win64">
 		<Node Name="Submit Baseline To Perforce" Requires="$(BaselineGraphNodes)">
-			<Submit Description="[CL $(Change)] Updated IterationProfile baselines" Files="$(BaselineFiles)"/>
+			<Submit Description="[CL $(Change)] Updated EditorPerf baselines" Files="$(BaselineFiles)"/>
 		</Node>
 	</Agent>
 	-->
diff --git a/Engine/Build/IterationProfile.xml b/Engine/Build/IterationProfile.xml
deleted file mode 100644
diff --git a/Engine/Source/Editor/UnrealEd/Private/Commandlets/SummarizeTraceCommandlet.cpp b/Engine/Source/Editor/UnrealEd/Private/Commandlets/SummarizeTraceCommandlet.cpp
--- a/Engine/Source/Editor/UnrealEd/Private/Commandlets/SummarizeTraceCommandlet.cpp
+++ b/Engine/Source/Editor/UnrealEd/Private/Commandlets/SummarizeTraceCommandlet.cpp
@@ -1274,14 +1274,14 @@ int32 USummarizeTraceCommandlet::Main(const FString& CmdLineParams)
 	// load the stats file to know which event name and statistic name to generate in the telementry csv
 	// the telemetry csv is ingested completely, so this just whitelists specific data elements we want to track
 	TMultiMap<FString, StatisticDefinition> NameToDefinitionMap;
-	FString GlobalStatisticsFileName = FPaths::RootDir() / TEXT("Engine") / TEXT("Build") / TEXT("IterationProfileStats.csv");
+	FString GlobalStatisticsFileName = FPaths::RootDir() / TEXT("Engine") / TEXT("Build") / TEXT("EditorPerfStats.csv");
 	if (FPaths::FileExists(GlobalStatisticsFileName))
 	{
 		UE_LOG(LogSummarizeTrace, Display, TEXT("Loading global statistics from %s"), *GlobalStatisticsFileName);
 		bool bCSVOk = StatisticDefinition::LoadFromCSV(GlobalStatisticsFileName, NameToDefinitionMap);
 		check(bCSVOk);
 	}
-	FString ProjectStatisticsFileName = FPaths::ProjectDir() / TEXT("Build") / TEXT("IterationProfileStats.csv");
+	FString ProjectStatisticsFileName = FPaths::ProjectDir() / TEXT("Build") / TEXT("EditorPerfStats.csv");
 	if (FPaths::FileExists(ProjectStatisticsFileName))
 	{
 		UE_LOG(LogSummarizeTrace, Display, TEXT("Loading project statistics from %s"), *ProjectStatisticsFileName);
