EditorPerf: Tidy up Agent Types for homogenized workspaces

#rnx

diff --git a/Engine/Build/EditorPerf.xml b/Engine/Build/EditorPerf.xml
--- a/Engine/Build/EditorPerf.xml
+++ b/Engine/Build/EditorPerf.xml
@@ -163,7 +163,7 @@
 		<ForEach Name="Project" Values="$(Projects)">
 
 			<!-- Project execution -->
-			<Agent Name="Automated Test Agent $(Project) $(Platform)" Type="EditorPerf$(Platform)$(Project)">
+			<Agent Name="Automated Test Agent $(Project) $(Platform)" Type="EditorPerf$(Platform)">
 				<Property Name="OutputName" Value="$(Project)$(Platform)EditorPerf"/>
 				<Property Name="RunNodeName" Value="Run $(OutputName)"/>
 				<Property Name="RunNodeRequires" Value="Compile Prerequisites $(Platform);Compile $(Project) Editor $(Platform)"/>
