EditorPerf: Add iteration suffix to phase for report key and artifact name

Skipping preflight as this change only effect EditorPerf builds

#rnx

diff --git a/Engine/Build/EditorPerf.xml b/Engine/Build/EditorPerf.xml
--- a/Engine/Build/EditorPerf.xml
+++ b/Engine/Build/EditorPerf.xml
@@ -94,8 +94,8 @@
 				<Property Name="TelemetryFileName" Value="$(IterationName)Telemetry.csv"/>
 				<Property Name="TelemetryFilePath" Value="$(TraceDir)/$(TelemetryFileName)"/>
 				<Property Name="LogOutputParams" Value="-logdir=&quot;$(AutomationOutputDir)/$(OutputName)&quot;"/>
-				<Property Name="ReportParams" Value="-WriteTestResultsForHorde -HordeTestDataKey=&quot;$(NodeName) $(Phase)$(PreflightSuffix)&quot;"/>
-				<Property Name="AllGauntletArgs" Value="-project=$(Project) -platform=$(Platform) $(GauntletArgs) $(ReportParams) $(LogOutputParams) -artifactname=&quot;$(Phase)&quot;"/>
+				<Property Name="ReportParams" Value="-WriteTestResultsForHorde -HordeTestDataKey=&quot;$(NodeName) $(Phase)$(Iteration)$(PreflightSuffix)&quot;"/>
+				<Property Name="AllGauntletArgs" Value="-project=$(Project) -platform=$(Platform) $(GauntletArgs) $(ReportParams) $(LogOutputParams) -artifactname=&quot;$(Phase)$(Iteration)&quot;"/>
 
 				<!-- Do DDC cleanup outside the editor process, handles Cold cleanup without changing stats -->
 				<Commandlet Project="$(Project)" Name="DDCCleanup" Arguments="$(Args) $(CommandletArgs)"/>
