Remove Lumin from Build files

diff --git a/Engine/Build/InstalledEngineBuild.xml b/Engine/Build/InstalledEngineBuild.xml
--- a/Engine/Build/InstalledEngineBuild.xml
+++ b/Engine/Build/InstalledEngineBuild.xml
@@ -42,8 +42,6 @@
 	<Option Name="WithTVOS" Restrict="true|false" DefaultValue="$(DefaultWithIOS)" Description="Include the tvOS target platform"/>
 	<Option Name="WithLinux" Restrict="true|false" DefaultValue="$(DefaultWithLinux)" Description="Include the Linux target platform"/>
 	<Option Name="WithLinuxArm64" Restrict="true|false" DefaultValue="$(DefaultWithLinuxArm64)" Description="Include the Linux AArch64 target platform"/>
-	<Option Name="WithLumin" Restrict="true|false" DefaultValue="false" Description="Include the Lumin target platform on Windows"/>
-	<Option Name="WithLuminMac" Restrict="true|false" DefaultValue="false" Description="Include the Lumin target platform on Mac"/>
 	<Option Name="WithHoloLens" Restrict="true|false" DefaultValue="$(DefaultWithHoloLens)" Description="Include the HoloLens target platform"/>
 
 	<!-- Set up target types -->
@@ -88,13 +86,11 @@
 
 	<!-- Architectures that we build for Android. NOTE: This should match the architectures specified in BuildPluginCommand. -->
 	<Property Name="AndroidArchitectures" Value="-armv7+-arm64"/>
-	<Property Name="LuminArchitectures" Value="armv7+arm64"/>
 	
 	<!-- Hololens Architectures -->
 	<Property Name="HoloLensArchitectures" Value="x64+arm64+multi"/>
 
 	<Property Name="AndroidGPUArchitectures" Value="es2"/>
-	<Property Name="LuminGPUArchitectures" Value="es2"/>
 	
 	<!-- Compile flags for Windows Game binaries -->
 	<Property Name="TargetDebugInfoArg" Value="-nodebuginfo"/>
@@ -345,21 +341,6 @@
 	</Agent>
 	<Aggregate Name="TargetPlatforms_HoloLens" Label="Platforms/HoloLens" Requires="Compile UnrealGame HoloLens"/>
 
-	<Agent Name="Target Lumin" Type="Win64_Licensee">
-		<Node Name="Compile UnrealGame Lumin" Requires="Compile UnrealHeaderTool Win64" Produces="#UnrealGame Lumin;#UnrealGame Lumin Unstripped;#UnrealGame Lumin Stripped">
-			<ForEach Name="Target" Values="UnrealGame;$(OptionalClientTarget)">
-				<ForEach Name="Configuration" Values="$(GameConfigurations)">
-					<Compile Target="$(Target)" Platform="Lumin" Configuration="$(Configuration)" Tag="#UnrealGame Lumin" AllowXGE="false" Arguments="-precompile -allmodules -nolink -architectures=$(LuminArchitectures) -gpuarchitectures=$(LuminGPUArchitectures) $(TargetDebugInfoArg)"/>
-					<Compile Target="$(Target)" Platform="Lumin" Configuration="$(Configuration)" Tag="#UnrealGame Lumin" AllowXGE="false" Arguments="-precompile -architectures=$(LuminArchitectures) -gpuarchitectures=$(LuminGPUArchitectures) $(TargetDebugInfoArg)" Clean="false"/>
-				</ForEach>
-			</ForEach>
-			
-			<Tag Files="#UnrealGame Lumin" Filter="$(LuminStripFilter)" With="#UnrealGame Lumin Unstripped"/>
-			<Strip BaseDir="$(RootDir)" OutputDir="$(SavedOutput)" Platform="Lumin" Files="#UnrealGame Lumin Unstripped" Tag="#UnrealGame Lumin Stripped"/>
-		</Node>
-	</Agent>
-	<Aggregate Name="TargetPlatforms_Lumin" Label="Platforms/Lumin" Requires="Compile UnrealGame Lumin"/>
-
 	<!-- Target Platforms normally built from a Mac machine -->
 	<Agent Name="Target Mac" Type="Mac_Licensee">
 		<Node Name="Compile UnrealGame Mac" Requires="Compile UnrealHeaderTool Mac" Produces="#UnrealGame Mac;#UnrealGame Mac Unstripped;#UnrealGame Mac Stripped;#UnrealGame Mac Unsigned;#UnrealGame Mac Signed">
@@ -412,21 +393,6 @@
 	</Agent>
 	<Aggregate Name="TargetPlatforms_TVOS" Label="Platforms/TVOS" Requires="Compile UnrealGame TVOS"/>
 
-	<Agent Name="Target LuminMac" Type="Mac_Licensee">
-		<Node Name="Compile UnrealGame LuminMac" Requires="Compile UnrealHeaderTool Mac" Produces="#UnrealGame LuminMac;#UnrealGame LuminMac Unstripped;#UnrealGame LuminMac Stripped">
-			<ForEach Name="Target" Values="UnrealGame;$(OptionalClientTarget)">
-				<ForEach Name="Configuration" Values="$(GameConfigurations)">
-					<Compile Target="$(Target)" Platform="Lumin" Configuration="$(Configuration)" Tag="#UnrealGame LuminMac" Arguments="-precompile -allmodules -nolink -architectures=$(LuminArchitectures) -gpuarchitectures=$(LuminGPUArchitectures) $(TargetDebugInfoArg)" AllowParallelExecutor="false"/>
-					<Compile Target="$(Target)" Platform="Lumin" Configuration="$(Configuration)" Tag="#UnrealGame LuminMac" Arguments="-precompile -architectures=$(LuminArchitectures) -gpuarchitectures=$(LuminGPUArchitectures) $(TargetDebugInfoArg)" AllowParallelExecutor="false" Clean="false"/>
-				</ForEach>
-			</ForEach>
-
-			<Tag Files="#UnrealGame LuminMac" Filter="$(LuminStripFilter)" With="#UnrealGame LuminMac Unstripped"/>
-			<Strip BaseDir="$(RootDir)" OutputDir="$(SavedOutput)" Platform="Lumin" Files="#UnrealGame LuminMac Unstripped" Tag="#UnrealGame LuminMac Stripped"/>
-		</Node>
-	</Agent>
-	<Aggregate Name="TargetPlatforms_LuminMac" Label="Platforms/LuminMac" Requires="Compile UnrealGame LuminMac"/>
-
 	<Agent Name="NotForLicensee Tools Group Win64" Type="Win64">
 		<Node Name="Build NotForLicensee Tools Win64" Requires="Compile UnrealHeaderTool Win64" Produces="#NotForLicensee Build Tools Win64">
 			<Compile Target="DatasmithCADWorker" Configuration="Development" Platform="Win64" Tag="#NotForLicensee Build Tools Win64"/>
@@ -616,7 +582,6 @@
 		<Property Name="DDCPlatformsWin64" Value="$(DDCPlatformsWin64)+TVOS" If="'$(WithTVOS)' == true"/>
 		<Property Name="DDCPlatformsWin64" Value="$(DDCPlatformsWin64)+Linux" If="'$(WithLinux)' == true"/>
 		<Property Name="DDCPlatformsWin64" Value="$(DDCPlatformsWin64)+LinuxArm64" If="'$(WithLinuxArm64)' == true"/>
-		<Property Name="DDCPlatformsWin64" Value="$(DDCPlatformsWin64)+Lumin" If="'$(WithLumin)' == true"/>
 		<Property Name="DDCPlatformsWin64" Value="$(DDCPlatformsWin64)+HoloLens" If="'$(WithHoloLens)' == true"/>
 		<ForEach Name="Platform" Values="$(ExtensionPlatforms)">
 			<Expand Name="$(Platform)_Append_DDCPlatformsWin64"/>
@@ -696,7 +661,6 @@
 		<Property Name="DDCPlatformsMac" Value="$(DDCPlatformsMac)+Android_ATC" If="'$(WithAndroid)' == true"/>
 		<Property Name="DDCPlatformsMac" Value="$(DDCPlatformsMac)+IOS" If="'$(WithIOS)' == true"/>
 		<Property Name="DDCPlatformsMac" Value="$(DDCPlatformsMac)+TVOS" If="'$(WithTVOS)' == true"/>
-		<Property Name="DDCPlatformsMac" Value="$(DDCPlatformsMac)+Lumin" If="'$(WithLuminMac)' == true"/>
 	</Do>
 
 	<Agent Name="DDC Group Mac" Type="Mac">
@@ -811,7 +775,6 @@
 		<Property Name="InstalledRequirements" Value="$(InstalledRequirements);#UnrealGame TVOS;#UnrealGame TVOS Unstripped;#UnrealGame TVOS Stripped" If="'$(WithTVOS)' == true"/>
 		<Property Name="InstalledRequirements" Value="$(InstalledRequirements);#UnrealGame Linux;#UnrealGame Linux Unstripped;#UnrealGame Linux Stripped;#Build Tools Linux" If="'$(WithLinux)' == true"/>
 		<Property Name="InstalledRequirements" Value="$(InstalledRequirements);#UnrealGame LinuxArm64;#UnrealGame LinuxArm64 Unstripped;#UnrealGame LinuxArm64 Stripped" If="'$(WithLinuxArm64)' == true"/>
-		<Property Name="InstalledRequirements" Value="$(InstalledRequirements);#UnrealGame Lumin" If="'$(WithLumin)' == true"/>
 		<Property Name="InstalledRequirements" Value="$(InstalledRequirements);#UnrealGame HoloLens" If="'$(WithHoloLens)' == true"/>
 		<Property Name="InstalledRequirements" Value="$(InstalledRequirements);#CompressedDDCWin64" If="$(WithDDC) And $(WithWin64)"/>
 		<Property Name="InstalledRequirements" Value="$(InstalledRequirements);#Compile Datasmith Plugins Win64;#DatasmithPluginBinaries Win64" If="$(WithWin64) And $(CompileDatasmithPlugins)"/>
@@ -884,12 +847,6 @@
 			<ForEach Name="Platform" Values="$(ExtensionPlatforms)">
 				<Expand Name="$(Platform)_Tag_UnrealGame"/>
 			</ForEach>
-			<Do If="'$(WithLumin)' == true">
-				<Tag Files="#UnrealGame Lumin" With="#Installed Win64"/>
-				<Property Name="Platforms" Value="$(Platforms)Lumin;"/>
-				<Property Name="CopyInstalledFilter" Value="$(CopyInstalledFilter);$(CopyLuminFilter)"/>
-				<Property Name="CopyInstalledExceptions" Value="$(CopyInstalledExceptions);$(CopyLuminExceptions)"/>
-			</Do>
 			<Do If="'$(WithHoloLens)' == true">
 				<Tag Files="#UnrealGame HoloLens;#UnrealGame HoloLens Includes" With="#Installed Win64"/>
 				<Property Name="Platforms" Value="$(Platforms)HoloLens;"/>
@@ -970,7 +927,6 @@
 		<Property Name="FilterRequirements" Value="$(FilterRequirements);#UnrealGame Android;#UnrealGame Android Unstripped;#UnrealGame Android Stripped" If="'$(WithAndroid)' == true"/>
 		<Property Name="FilterRequirements" Value="$(FilterRequirements);#UnrealGame IOS;#UnrealGame IOS Unstripped;#UnrealGame IOS Stripped" If="'$(WithIOS)' == true"/>
 		<Property Name="FilterRequirements" Value="$(FilterRequirements);#UnrealGame TVOS;#UnrealGame TVOS Unstripped;#UnrealGame TVOS Stripped" If="'$(WithTVOS)' == true"/>
-		<Property Name="FilterRequirements" Value="$(FilterRequirements);#UnrealGame LuminMac;#UnrealGame LuminMac Unstripped;#UnrealGame LuminMac Stripped" If="'$(WithLuminMac)' == true"/>
 		<Property Name="FilterRequirements" Value="$(FilterRequirements);#CompressedDDCMac" If="$(WithDDC) And $(WithMac)"/>
 		<Property Name="FilterRequirements" Value="$(FilterRequirements);#Compile Datasmith Plugins Mac;#DatasmithPluginBinaries Mac" If="$(WithMac) And $(CompileDatasmithPlugins)"/>
 
@@ -1025,13 +981,6 @@
 				<Property Name="CopyInstalledFilter" Value="$(CopyInstalledFilter);$(CopyTVOSFilter)"/>
 				<Property Name="CopyInstalledExceptions" Value="$(CopyInstalledExceptions);$(CopyTVOSExceptions)"/>
 			</Do>
-			<Do If="'$(WithLuminMac)' == true">
-				<Tag Files="#UnrealGame LuminMac" Except="#UnrealGame LuminMac Unstripped" With="#Installed Mac"/>
-				<Tag Files="#UnrealGame LuminMac Stripped" With="#Saved Output"/>
-				<Property Name="Platforms" Value="$(Platforms)Lumin;"/>
-				<Property Name="CopyInstalledFilter" Value="$(CopyInstalledFilter);$(CopyLuminFilter)"/>
-				<Property Name="CopyInstalledExceptions" Value="$(CopyInstalledExceptions);$(CopyLuminExceptions)"/>
-			</Do>
 
 			<!-- Tag any dependencies from all previous build products -->
 			<Tag Files="#Installed Mac" Filter="*.target" With="#TargetReceipts"/>
