EditorPerf: Promote Iterations property to option

#rnx

diff --git a/Engine/Build/EditorPerf.xml b/Engine/Build/EditorPerf.xml
--- a/Engine/Build/EditorPerf.xml
+++ b/Engine/Build/EditorPerf.xml
@@ -61,7 +61,7 @@
 	<Option Name="WarmArgs" DefaultValue="-ddc=Warm" Description="Arguments to unreal for the warm phase"/>
 
 	<!-- Iterations to run and average together -->
-	<Property Name="Iterations" Value="1;2;3"/>
+	<Option Name="Iterations" DefaultValue="1;2;3" Description="The number of iterations to run"/>
 
 	<!-- Gauntlet arguments for all runs -->
 	<Property Name="GauntletArgs" Value="-build=editor -UseEditor -verbose -tempdir=&quot;$(RootDir)/Tests&quot;"/>
