Split Linux platform into another grouping LinuxNative for UnrealGame, Build Tools, Feature Pack.

This include LinuxArm64 UnrealGame, but not Android UnrealGame due to no AutoSDK for Android for Linux atm

This will only target native linux horde or machines. Moving them to Linux_License as well.

diff --git a/Engine/Build/InstalledEngineBuild.xml b/Engine/Build/InstalledEngineBuild.xml
--- a/Engine/Build/InstalledEngineBuild.xml
+++ b/Engine/Build/InstalledEngineBuild.xml
@@ -218,13 +218,13 @@
 	<Aggregate Name="Mac Editor" Label="Editors/Mac" Requires="Compile UnrealEditor Mac" Include="Editor Mac"/>
 
 	<!-- Base editor and tools for compile on Linux -->
-	<Agent Name="Editor Linux" Type="Win64_Licensee">
+	<Agent Name="Editor Linux" Type="Linux_Licensee">
 		<Node Name="Compile UnrealHeaderTool Linux" Requires="Update Version Files">
 			<Compile Target="UnrealHeaderTool" Platform="Linux" Configuration="Development" Arguments="-precompile -allmodules"/>
 		</Node>
-		<Node Name="Compile UnrealEditor Linux" Requires="Compile UnrealHeaderTool $(HostPlatform)" Produces="#UnrealEditor Linux;#UnrealEditor Linux Unstripped;#UnrealEditor Linux Stripped">
-			<Compile Target="UnrealEditor" Platform="Linux" Configuration="DebugGame" Tag="#UnrealEditor Linux" Arguments="-precompile -allmodules"/>
-			<Compile Target="UnrealEditor" Platform="Linux" Configuration="Development" Tag="#UnrealEditor Linux" Arguments="-precompile -allmodules"/>
+		<Node Name="Compile UnrealEditor Linux" Requires="Compile UnrealHeaderTool Linux" Produces="#UnrealEditor Linux;#UnrealEditor Linux Unstripped;#UnrealEditor Linux Stripped">
+			<Compile Target="UnrealEditor" Platform="Linux" Configuration="DebugGame" Tag="#UnrealEditor Linux" Arguments="-precompile -allmodules $(CrashReporterCompileArgs)"/>
+			<Compile Target="UnrealEditor" Platform="Linux" Configuration="Development" Tag="#UnrealEditor Linux" Arguments="-precompile -allmodules $(CrashReporterCompileArgs)"/>
 
 			<Tag Files="#UnrealEditor Linux" Filter="$(LinuxStripFilter)" With="#UnrealEditor Linux Unstripped"/>
 
@@ -240,6 +240,7 @@
 			</Switch>
 		</Node>
 	</Agent>
+	<Aggregate Name="Linux Editor Linux" Label="Editors/Linux" Requires="Compile UnrealEditor Linux" Include="Editor Linux"/>
 
 	<!-- Target Platforms normally built from a Win64 machine -->
 	<Agent Name="Target Win64" Type="Win64_Licensee">
@@ -281,7 +282,7 @@
 	<Aggregate Name="TargetPlatforms_Win64" Requires="Compile UnrealGame Win64"/>
 
 	<Agent Name="Target Linux" Type="Win64_Licensee">
-		<Node Name="Compile UnrealGame Linux" Requires="Compile UnrealHeaderTool $(HostPlatform)" Produces="#UnrealGame Linux;#UnrealGame Linux Unstripped;#UnrealGame Linux Stripped">
+		<Node Name="Compile UnrealGame Linux" Requires="Compile UnrealHeaderTool Win64" Produces="#UnrealGame Linux;#UnrealGame Linux Unstripped;#UnrealGame Linux Stripped">
 			<ForEach Name="Target" Values="UnrealGame;$(OptionalClientTarget);$(OptionalServerTarget)">
 				<ForEach Name="Configuration" Values="$(GameConfigurations)">
 					<Compile Target="$(Target)" Platform="Linux" Configuration="$(Configuration)" Tag="#UnrealGame Linux" Arguments="-precompile -allmodules -nolink $(TargetDebugInfoArg)"/>
@@ -294,9 +295,24 @@
 		</Node>
 	</Agent>
 	<Aggregate Name="TargetPlatforms_Linux" Requires="Compile UnrealGame Linux"/>
-	
+
+	<Agent Name="Target LinuxNative" Type="Linux_Licensee">
+		<Node Name="Compile UnrealGame LinuxNative" Requires="Compile UnrealHeaderTool Linux" Produces="#UnrealGame LinuxNative;#UnrealGame LinuxNative Unstripped;#UnrealGame LinuxNative Stripped">
+			<ForEach Name="Target" Values="UnrealGame;$(OptionalClientTarget);$(OptionalServerTarget)">
+				<ForEach Name="Configuration" Values="$(GameConfigurations)">
+					<Compile Target="$(Target)" Platform="Linux" Configuration="$(Configuration)" Tag="#UnrealGame LinuxNative" Arguments="-precompile -allmodules -nolink $(TargetDebugInfoArg)"/>
+					<Compile Target="$(Target)" Platform="Linux" Configuration="$(Configuration)" Tag="#UnrealGame LinuxNative" Arguments="-precompile $(TargetDebugInfoArg)" Clean="false"/>
+				</ForEach>
+			</ForEach>
+
+			<Tag Files="#UnrealGame LinuxNative Linux" Filter="$(LinuxStripFilter)" With="#UnrealGame LinuxNative Unstripped"/>
+			<Strip BaseDir="$(RootDir)" OutputDir="$(SavedOutput)" Platform="Linux" Files="#UnrealGame LinuxNative Unstripped" Tag="#UnrealGame LinuxNative Stripped"/>
+		</Node>
+	</Agent>
+	<Aggregate Name="TargetPlatforms_LinuxNative" Requires="Compile UnrealGame LinuxNative"/>
+
 	<Agent Name="Target LinuxArm64" Type="Win64_Licensee">
-		<Node Name="Compile UnrealGame LinuxArm64" Requires="Compile UnrealHeaderTool $(HostPlatform)" Produces="#UnrealGame LinuxArm64;#UnrealGame LinuxArm64 Unstripped;#UnrealGame LinuxArm64 Stripped">
+		<Node Name="Compile UnrealGame LinuxArm64" Requires="Compile UnrealHeaderTool Win64" Produces="#UnrealGame LinuxArm64;#UnrealGame LinuxArm64 Unstripped;#UnrealGame LinuxArm64 Stripped">
 			<ForEach Name="Target" Values="UnrealGame;$(OptionalClientTarget);$(OptionalServerTarget)">
 				<ForEach Name="Configuration" Values="$(GameConfigurations)">
 					<Compile Target="$(Target)" Platform="LinuxArm64" Configuration="$(Configuration)" Tag="#UnrealGame LinuxArm64" Arguments="-precompile -allmodules -nolink $(TargetDebugInfoArg)"/>
@@ -310,6 +326,21 @@
 	</Agent>
 	<Aggregate Name="TargetPlatforms_LinuxArm64" Label="Platforms/Linux (AArch64)" Requires="Compile UnrealGame LinuxArm64"/>
 
+	<Agent Name="Target LinuxArm64 LinuxNative" Type="Linux_Licensee">
+		<Node Name="Compile UnrealGame LinuxArm64 LinuxNative" Requires="Compile UnrealHeaderTool Linux" Produces="#UnrealGame LinuxArm64 LinuxNative;#UnrealGame LinuxArm64 LinuxNative Unstripped;#UnrealGame LinuxArm64 LinuxNative Stripped">
+			<ForEach Name="Target" Values="UnrealGame;$(OptionalClientTarget);$(OptionalServerTarget)">
+				<ForEach Name="Configuration" Values="$(GameConfigurations)">
+					<Compile Target="$(Target)" Platform="LinuxArm64" Configuration="$(Configuration)" Tag="#UnrealGame LinuxArm64 LinuxNative" Arguments="-precompile -allmodules -nolink $(TargetDebugInfoArg)"/>
+					<Compile Target="$(Target)" Platform="LinuxArm64" Configuration="$(Configuration)" Tag="#UnrealGame LinuxArm64 LinuxNative" Arguments="-precompile $(TargetDebugInfoArg)" Clean="false"/>
+				</ForEach>
+			</ForEach>
+
+			<Tag Files="#UnrealGame LinuxArm6 LinuxNative" Filter="$(LinuxStripFilter)" With="#UnrealGame LinuxArm64 LinuxNative Unstripped"/>
+			<Strip BaseDir="$(RootDir)" OutputDir="$(SavedOutput)" Platform="LinuxArm64" Files="#UnrealGame LinuxArm64 LinuxNative Unstripped" Tag="#UnrealGame LinuxArm64 LinuxNative Stripped"/>
+		</Node>
+	</Agent>
+	<Aggregate Name="TargetPlatforms_LinuxArm64LinuxNative" Label="Platforms/Linux (AArch64)" Requires="Compile UnrealGame LinuxArm64 LinuxNative"/>
+
 	<Agent Name="Target Android" Type="Win64_Licensee">
 		<Node Name="Compile UnrealGame Android" Requires="Compile UnrealHeaderTool $(HostPlatform)" Produces="#UnrealGame Android;#UnrealGame Android Unstripped;#UnrealGame Android Stripped">
 			<ForEach Name="Target" Values="UnrealGame;$(OptionalClientTarget)">
@@ -324,7 +355,7 @@
 		</Node>
 	</Agent>
 	<Aggregate Name="TargetPlatforms_Android" Label="Platforms/Android" Requires="Compile UnrealGame Android"/>
-	
+
 	<ForEach Name="Platform" Values="$(ExtensionPlatforms)">
 		<Expand Name="Agent_Target_$(Platform)_Win64_Licensee"/>
 	</ForEach>
@@ -485,7 +516,7 @@
 			<!-- Sign the binaries -->
 			<Sign Files="#Sign Binaries" If="'$(SignExecutables)' == true"/>
 		</Node>
-		<Node Name="Build Tools Linux" Requires="Compile UnrealHeaderTool $(HostPlatform)">
+		<Node Name="Build Tools Linux" Requires="Compile UnrealHeaderTool Win64">
 			<Compile Target="CrashReportClient" Platform="Linux" Configuration="Shipping"/>
 			<Compile Target="CrashReportClientEditor" Platform="Linux" Configuration="Shipping"/>
 			<Compile Target="ShaderCompileWorker" Configuration="Development" Platform="Linux" Tag="#Build Tools Linux"/>
@@ -561,6 +592,48 @@
 		</Node>
 	</Agent>
 
+	<Agent Name="Tools Group LinuxNative" Type="Linux_Licensee">
+		<Node Name="Build Tools LinuxNative" Requires="Compile UnrealHeaderTool Linux">
+			<Compile Target="CrashReportClient" Platform="Linux" Configuration="Shipping"/>
+			<Compile Target="CrashReportClientEditor" Platform="Linux" Configuration="Shipping"/>
+			<Compile Target="ShaderCompileWorker" Configuration="Development" Platform="Linux" Tag="#Build Tools LinuxNative"/>
+			<Compile Target="UnrealInsights" Configuration="Development" Platform="Linux" Tag="#Build Tools LinuxNative"/>
+			<Compile Target="UnrealLightmass" Configuration="Development" Platform="Linux" Tag="#Build Tools LinuxNative"/>
+			<Compile Target="UnrealPak" Configuration="Development" Platform="Linux" Tag="#Build Tools LinuxNative"/>
+			<Compile Target="UnrealMultiUserServer" Configuration="Development" Platform="Linux" Tag="#Build Tools LinuxNative"/>
+			<Compile Target="UnrealRecoverySvc" Configuration="Development" Platform="Linux" Tag="#Build Tools LinuxNative"/>
+			<Compile Target="SwitchboardListener" Configuration="Development" Platform="Linux" Tag="#Build Tools LinuxNative"/>
+		</Node>
+
+		<!-- Compile the C# tools. Copies files to a temporary directory first, to strip out all confidential folders and avoid sharing violations with the running version of UAT. -->
+		<Node Name="Build Tools CS LinuxNative" Requires="Update Version Files" Produces="#Build Tools CS Binaries LinuxNative">
+			<Delete Files="$(CsToolsDir)/..."/>
+
+			<!-- Copy Source and referenced libraries to a new location with Confidential folders removed-->
+			<Tag Files="$(AutomationToolPaths)" Filter="*.csproj" Except="$(ConfidentialExceptions)"  With="#UAT Projects"/>
+			<Tag Files="Engine/Binaries/DotNET/...;Engine/Binaries/ThirdParty/Newtonsoft/...;Engine/Binaries/ThirdParty/IOS/...;Engine/Binaries/ThirdParty/VisualStudio/...;Engine/Source/Programs/...;Engine/Platforms/*/Source/Programs/...;Engine/Source/Editor/SwarmInterface/...;#UATReferences" Except="$(ConfidentialExceptions);AutomationTool.csproj.References" With="#RedistUATSource"/>
+			<ForEach Name="Platform" Values="$(ExtensionPlatforms)">
+				<Expand Name="$(Platform)_TagCsToolsRedistUATSource"/>
+			</ForEach>
+			<Copy Files="#RedistUATSource" From="$(RootDir)" To="$(CsToolsDir)"/>
+
+			<Property Name="RedistAutomationToolPaths" Value=""/>
+			<ForEach Name="AutomationToolPath" Values="$(AutomationToolPaths)">
+				<Property Name="RedistAutomationToolPaths" Value="$(RedistAutomationToolPaths)$(CsToolsDir)/$(AutomationToolPath);"/>
+			</ForEach>
+			<Tag Files="$(RedistAutomationToolPaths)" Filter="*.csproj" With="#RedistUATProjects"/>
+
+			<!-- Tag AutomationTool and UnrealBuildTool folders recursively as NET Core dependencies are not currently handled by CsCompile -->
+			<Tag Files="$(CsToolsDir)/Engine/Binaries/DotNET/AutomationTool/..." With="#Build Tools CS Binaries LinuxNative" />
+			<Tag Files="$(CsToolsDir)/Engine/Binaries/DotNET/UnrealBuildTool/..." With="#Build Tools CS Binaries LinuxNative" />
+
+			<CsCompile Project="$(CsToolsDir)/Engine/Source/Programs/IOS/MobileDeviceInterface/MobileDeviceInterface.csproj" Configuration="Development" Platform="AnyCPU" Arguments="/verbosity:minimal /target:Rebuild" Tag="#Build Tools CS Binaries LinuxNative;#Sign Binaries LinuxNative" TagReferences="#Build Tools CS Binaries LinuxNative"/>
+
+			<!-- Sign the binaries -->
+			<Sign Files="#Sign Binaries LinuxNative" If="'$(SignExecutables)' == true"/>
+		</Node>
+	</Agent>
+
 	<Agent Name="FeaturePacks Group" Type="Win64_Licensee">
 		<Node Name="Make Feature Packs" Requires="Update Version Files;Build Tools $(HostPlatform)">
 			<!--Delete any pre-existing feature packs-->
@@ -574,6 +647,19 @@
 		</Node>
 	</Agent>
 
+	<Agent Name="FeaturePacks Group LinuxNative" Type="Linux_Licensee">
+		<Node Name="Make Feature Packs LinuxNative" Requires="Update Version Files;Build Tools LinuxNative">
+			<!--Delete any pre-existing feature packs-->
+			<Delete Files="FeaturePacks/..."/>
+			<ForEach Name="TemplateName" Values="$(ProjectsToFeaturePack)">
+				<PakFile Output="FeaturePacks/$(TemplateName).upack" Files="" ResponseFile="Templates/$(TemplateName)/contents.txt"/>
+			</ForEach>
+			<ForEach Name="SampleName" Values="StarterContent">
+				<PakFile Output="FeaturePacks/$(SampleName).upack" Files="" ResponseFile="Samples/$(SampleName)/contents.txt"/>
+			</ForEach>
+		</Node>
+	</Agent>
+
 	<Property Name="DDCPlatformsWin64" Value="WindowsEditor"/>
 	<Property Name="DDCPlatformsWin64" Value="$(DDCPlatformsWin64)+Windows" If="'$(WithWin64)' == true"/>
 	<Do If="'$(HostPlatformDDCOnly)' == false">
@@ -698,10 +784,8 @@
 		<Property Name="DDCPlatformsLinux" Value="$(DDCPlatformsLinux)+Android_ATC" If="'$(WithAndroid)' == true"/>
 	</Do>
 
-	<Agent Name="DDC Group Linux" Type="Linux;Win64">
-		<!-- Additionally require the Win64 editor binaries if building on win64 -->
-		<Property Name="DDCLinuxRequires" Value="#UnrealEditor Linux;Build Tools Linux;#Build Tools CS Binaries"/>
-		<Property Name="DDCLinuxRequires" Value="$(DDCLinuxRequires);#UnrealEditor Win64;#Build Tools Win64" If="'$(HostPlatform)' == 'Win64'"/>
+	<Agent Name="DDC Group Linux" Type="Linux">
+		<Property Name="DDCLinuxRequires" Value="#UnrealEditor Linux;Build Tools LinuxNative;#Build Tools CS Binaries LinuxNative"/>
 		<Node Name="Build DDC Linux" Requires="$(DDCLinuxRequires)" Produces="#CompressedDDCLinux">
 			<!-- Get our temp dir -->
 			<Property Name="DDCDir" Value="$(RootDir)/LocalBuilds/InstalledDDC"/>
@@ -710,8 +794,7 @@
 			<Delete Files="$(DDCDir)/..."/>
 			<!-- Build up a list of files needed to build DDC -->
 			<CsCompile Project="Engine/Source/Programs/UnrealBuildTool/UnrealBuildTool.csproj" Configuration="Development" Platform="AnyCPU" Tag="#ToCopy" TagReferences="#ToCopy" EnumerateOnly="true"/>
-			<Tag Files="#Compile UnrealHeaderTool $(HostPlatform);#UnrealEditor Linux;#Build Tools Linux" With="#ToCopy"/>
-			<Tag Files="#UnrealEditor Win64;#Build Tools Win64" With="#ToCopy" If="'$(HostPlatform)' == 'Win64'"/>
+			<Tag Files="#Compile UnrealHeaderTool Linux;#UnrealEditor Linux;#Build Tools LinuxNative" With="#ToCopy"/>
 			<Tag Files="#ToCopy" Filter="*.target" With="#TargetReceipts"/>
 			<TagReceipt Files="#TargetReceipts" RuntimeDependencies="true" With="#ToCopy"/>
 			<Tag Files="Engine/Content/..." Except="*.psd;*.pdn;*.fbx;*.po" With="#ToCopy"/>
@@ -724,10 +807,10 @@
 			<Expand Name="TagDDCEngineFiles" ExceptPlugins="Plugins/.../Binaries/Mac/..."/>
 			<!-- Filter out the files not needed to build DDC. Removing confidential folders can affect DDC keys, so we want to be sure that we're making DDC with a build that can use it. -->
 			<Tag Files="#ToCopy" Except="$(ConfidentialExceptions);$(RootDir)/.../Source/...;$(RootDir)/.../Intermediate/..." With="#FilteredCopyList"/>
-			
+
 			<!-- Copy everything to a temporary directory -->
 			<Copy From="$(RootDir)" To="$(DDCDir)" Files="#FilteredCopyList"/>
-			<Copy From="$(CsToolsDir)" To="$(DDCDir)" Files="#Build Tools CS Binaries"/>
+			<Copy From="$(CsToolsDir)" To="$(DDCDir)" Files="#Build Tools CS Binaries LinuxNative"/>
 
 			<Command Name="BuildDerivedDataCache" Arguments="-TempDir=&quot;$(DDCDir)&quot; -FeaturePacks=$(ProjectsToBuildDDC) -TargetPlatforms=$(DDCPlatformsLinux) -SavedDir=&quot;$(SavedDir)&quot;"/>
 
@@ -1040,18 +1123,18 @@
 	</Agent>
 	<Aggregate Name="HostPlatforms_Mac" Label="Builds/Mac" Requires="Make Installed Build Mac"/>
 
-	<Agent Name="Installed Build Group Linux" Type="Win64_Licensee">
+	<Agent Name="Installed Build Group Linux" Type="Linux_Licensee">
 		<!-- Build up a list of requirements as some can be optional -->
-		<Property Name="InstalledRequirements" Value="#Compile UnrealHeaderTool Linux;#UnrealEditor Linux;#UnrealEditor Linux Unstripped;#UnrealEditor Linux Stripped;#Build Tools Linux;#Build Tools CS Binaries;#Rules Assemblies"/>
+		<Property Name="InstalledRequirements" Value="#Compile UnrealHeaderTool Linux;#UnrealEditor Linux;#UnrealEditor Linux Unstripped;#UnrealEditor Linux Stripped;#Build Tools LinuxNative;#Build Tools CS Binaries LinuxNative;#Rules Assemblies"/>
 
 		<!-- Optional Target Platform requirements -->
-		<Property Name="InstalledRequirements" Value="$(InstalledRequirements);#UnrealGame Linux;#UnrealGame Linux Unstripped;#UnrealGame Linux Stripped" If="'$(WithLinux)' == true"/>
-		<Property Name="InstalledRequirements" Value="$(InstalledRequirements);#UnrealGame LinuxArm64;#UnrealGame LinuxArm64 Unstripped;#UnrealGame LinuxArm64 Stripped" If="'$(WithLinuxArm64)' == true"/>
-		<Property Name="InstalledRequirements" Value="$(InstalledRequirements);#UnrealGame Android;#UnrealGame Android Unstripped;#UnrealGame Android Stripped" If="'$(WithAndroid)' == true"/>		
+		<Property Name="InstalledRequirements" Value="$(InstalledRequirements);#UnrealGame LinuxNative;#UnrealGame LinuxNative Unstripped;#UnrealGame LinuxNative Stripped" If="'$(WithLinux)' == true"/>
+		<Property Name="InstalledRequirements" Value="$(InstalledRequirements);#UnrealGame LinuxArm64 LinuxNative;#UnrealGame LinuxArm64 LinuxNative Unstripped;#UnrealGame LinuxArm64 LinuxNative Stripped" If="'$(WithLinuxArm64)' == true"/>
+		<Property Name="InstalledRequirements" Value="$(InstalledRequirements);#UnrealGame Android;#UnrealGame Android Unstripped;#UnrealGame Android Stripped" If="'$(WithAndroid)' == true"/>
 		<Property Name="InstalledRequirements" Value="$(InstalledRequirements);#CompressedDDCLinux" If="'$(WithDDC)' == true And $(WithLinux)"/>
 
 		<!-- Feature packs -->
-		<Property Name="InstalledRequirements" Value="$(InstalledRequirements);Make Feature Packs"/>
+		<Property Name="InstalledRequirements" Value="$(InstalledRequirements);Make Feature Packs LinuxNative"/>
 
 		<!-- Filters the build products required for an installed build into groups depending on whether they need to be stripped or just copied -->
 		<Node Name="Make Installed Build Linux" Requires="$(InstalledRequirements)" Produces="#Installed Build Linux Files">
@@ -1070,19 +1153,19 @@
 			<Tag Files="#Compile UnrealHeaderTool Linux" With="#Installed Linux"/>
 			<Tag Files="#UnrealEditor Linux" Except="#UnrealEditor Linux Unstripped" With="#Installed Linux"/>
 			<Tag Files="#UnrealEditor Linux Stripped" With="#Saved Output"/>
-			<Tag Files="#Build Tools Linux;#Make Feature Packs;#Rules Assemblies" With="#Installed Linux"/>
+			<Tag Files="#Build Tools LinuxNative;#Make Feature Packs LinuxNative;#Rules Assemblies" With="#Installed Linux"/>
 
 			<!-- Tag enabled target platforms -->
 			<Do If="'$(WithLinux)' == true">
-				<Tag Files="#UnrealGame Linux;#Build Tools Linux" Except="#UnrealGame Linux Unstripped" With="#Installed Linux"/>
-				<Tag Files="#UnrealGame Linux Stripped" With="#Saved Output"/>
+				<Tag Files="#UnrealGame LinuxNative;#Build Tools LinuxNative" Except="#UnrealGame LinuxNative Unstripped" With="#Installed Linux"/>
+				<Tag Files="#UnrealGame LinuxNative Stripped" With="#Saved Output"/>
 				<Property Name="Platforms" Value="$(Platforms)Linux;"/>
 				<Property Name="CopyInstalledFilter" Value="$(CopyInstalledFilter);$(CopyLinuxFilter)"/>
 				<Property Name="CopyInstalledExceptions" Value="$(CopyInstalledExceptions);$(CopyLinuxExceptions)"/>
 			</Do>
 			<Do If="'$(WithLinuxArm64)' == true">
-				<Tag Files="#UnrealGame LinuxArm64" Except="#UnrealGame LinuxArm64 Unstripped" With="#Installed Linux"/>
-				<Tag Files="#UnrealGame LinuxArm64 Stripped" With="#Saved Output"/>
+				<Tag Files="#UnrealGame LinuxArm64 LinuxNative" Except="#UnrealGame LinuxArm64 LinuxNative Unstripped" With="#Installed Linux"/>
+				<Tag Files="#UnrealGame LinuxArm64 LinuxNative Stripped" With="#Saved Output"/>
 				<Property Name="Platforms" Value="$(Platforms)LinuxArm64;"/>
 				<Property Name="CopyInstalledFilter" Value="$(CopyInstalledFilter);$(CopyLinuxFilter)"/>
 				<Property Name="CopyInstalledExceptions" Value="$(CopyInstalledExceptions);$(CopyLinuxExceptions)"/>
@@ -1110,7 +1193,7 @@
 			<Tag Files="$(CopyInstalledFilter)" With="#Installed Linux"/>
 
 			<!-- Tag Linux files that need to be stripped -->
-			<Tag Files="#Compile UnrealHeaderTool Linux;#Build Tools Linux" Filter="$(LinuxStripFilter)" Except="$(CopyInstalledExceptions)" With="#Strip Tools for Installed Linux"/>
+			<Tag Files="#Compile UnrealHeaderTool Linux;#Build Tools LinuxNative" Filter="$(LinuxStripFilter)" Except="$(CopyInstalledExceptions)" With="#Strip Tools for Installed Linux"/>
 
 			<!-- Tag files that need to be copied directly -->
 			<Tag Files="#Installed Linux" Except="$(CopyInstalledExceptions);#Strip Tools for Installed Linux" With="#Copy for Installed Linux"/>
@@ -1119,7 +1202,7 @@
 			<!-- Tag saved output that needs to be copied directly -->
 			<Tag BaseDir="$(SavedOutput)" Files="#Saved Output" Except="$(CopyInstalledExceptions)" With="#Copy Saved Output"/>
 			<Copy Files="#Copy Saved Output" From="$(SavedOutput)" To="$(LocalInstalledDirLinux)" Tag="#Installed Build Linux Files"/>
-			<Copy Files="#Build Tools CS Binaries" From="$(CsToolsDir)" To="$(LocalInstalledDirLinux)" Tag="#Installed Build Linux Files"/>
+			<Copy Files="#Build Tools CS Binaries LinuxNative" From="$(CsToolsDir)" To="$(LocalInstalledDirLinux)" Tag="#Installed Build Linux Files"/>
 			
 			<!-- Strip required files into their final location -->
 			<Strip Files="#Strip Tools for Installed Linux" BaseDir="$(RootDir)" OutputDir="$(LocalInstalledDirLinux)" Platform="Linux" Tag="#Installed Build Linux Files"/>
