EditorPerf

Update project configurations, projects configured in horde
Use to MapsToPIETest instead of AssetsToOpen... because we PIE test them :)

#rnx

diff --git a/Engine/Build/EditorPerf.xml b/Engine/Build/EditorPerf.xml
--- a/Engine/Build/EditorPerf.xml
+++ b/Engine/Build/EditorPerf.xml
@@ -30,7 +30,8 @@
 		<Property Name="ExtraPluginString" Value="$(ExtraPluginsPlusDelims)+$(Plugin)" If="'$(ExtraPluginsPlusDelims)' != ''"/>
 		<Property Name="ExtraPluginsPlusDelims" Value="$(ExtraPluginString)"/>
 	</ForEach>
-	<Property Name="EditorCompileArgs" Value="-BuildPlugin=$(ExtraPluginsPlusDelims)"/>
+	<Property Name="EditorCompileArgs" Value=""/>
+	<Property Name="EditorCompileArgs" Value="-BuildPlugin=$(ExtraPluginsPlusDelims)" If="'$(ExtraPluginsPlusDelims)' != ''"/>
 
 	<!-- Automation uses its own network output directory -->
 	<Property Name="AutomationNetworkRootDir" Value="\\epicgames.net\root\Builds\Automation\Reports\EditorPerf"/>
@@ -69,7 +70,7 @@
 	<Property Name="GauntletArgs" Value="$(GauntletArgs) -maxduration=28800"/>
 
 	<!-- Gauntlet arguments for profile runs -->
-	<Property Name="ProfileArgs" Value="-test=EditorTests -testfilter=Group:Iteration -tracefiletrunc"/>
+	<Property Name="ProfileArgs" Value="-test=EditorTests -testfilter=Project.Iteration.PIE -tracefiletrunc"/>
 	<Property Name="ProfileArgs" Value="$(ProfileArgs) -EnablePlugins=$(ExtraPluginsCommaDelims)" If="'$(ExtraPluginsCommaDelims)' != ''"/>
 	<Property Name="ProfileArgs" Value="$(ProfileArgs) -SuppressLogWarnings -SuppressLogErrors"/>
 	<Property Name="ProfileArgs" Value="$(ProfileArgs) -GameInstanceLostTimerSeconds=3600"/>
@@ -95,7 +96,7 @@
 				<Property Name="TelemetryFilePath" Value="$(TraceDir)/$(TelemetryFileName)"/>
 				<Property Name="LogOutputParams" Value="-logdir=&quot;$(AutomationOutputDir)/$(OutputName)&quot;"/>
 				<Property Name="ReportParams" Value="-WriteTestResultsForHorde -HordeTestDataKey=&quot;$(NodeName) $(Phase)$(Iteration)$(PreflightSuffix)&quot;"/>
-				<Property Name="AllGauntletArgs" Value="-project=$(Project) -platform=$(Platform) $(GauntletArgs) $(ReportParams) $(LogOutputParams) -artifactname=&quot;$(Phase)$(Iteration)&quot;"/>
+				<Property Name="AllGauntletArgs" Value="-Project=$(Project) -Platform=$(Platform) $(GauntletArgs) $(ReportParams) $(LogOutputParams) -artifactname=&quot;$(Phase)$(Iteration)&quot;"/>
 
 				<!-- Do DDC cleanup outside the editor process, handles Cold cleanup without changing stats -->
 				<Commandlet Project="$(Project)" Name="DDCCleanup" Arguments="$(Args) $(CommandletArgs)"/>
@@ -142,8 +143,10 @@
 	<!-- ForEach Platform to run, generate graph nodes -->
 	<ForEach Name="Platform" Values="$(Platforms)">
 
-		<!-- Project-independent prereq compiles -->
-		<Agent Name="Prerequisites Agent $(Platform)" Type="$($(Platform)CompileAgentType)">
+		<!-- We do everything on a single resource because it's a beast and we don't want to wait or use caching/distribution because it will pollute perf numbers -->
+		<Agent Name="Automated Test Agent $(Platform)" Type="EditorPerf$(Platform)">
+
+			<!-- Build Platform-global prerequisite targets -->
 			<Node Name="Compile Prerequisites $(Platform)">
 				<Compile Target="UnrealHeaderTool" Platform="$(Platform)" Configuration="Development"/>
 				<Compile Target="ShaderCompileWorker" Platform="$(Platform)" Configuration="Development"/>
@@ -151,19 +154,14 @@
 				<Compile Target="CrashReportClientEditor" Platform="$(Platform)" Configuration="Shipping"/>
 			</Node>
 
-			<!-- ForEach Project to run, generate the compile graph nodes -->
+			<!-- ForEach Project to run, generate the graph nodes -->
 			<ForEach Name="Project" Values="$(Projects)">
-				<Node Name="Compile $(Project) Editor $(Platform)" Requires="Compile Prerequisites $(Platform)">
-					<Compile Target="$(Project)Editor" Platform="$(Platform)" Configuration="Development" Arguments="$(EditorCompileArgs)"/>
-				</Node>
-			</ForEach>
-		</Agent>
 
-		<!-- ForEach Project to run, generate the graph nodes -->
-		<ForEach Name="Project" Values="$(Projects)">
+				<!-- We stuff the Editor target name in [] following the project name in case it's not $(Project)Editor -->
+				<Property Name="EditorTarget" Value="$(Project)Editor"/>
+				<Regex Pattern="(.*)\[(.*)\]" Capture="Project;EditorTarget" Input="$(Project)" Optional="true"/>
 
-			<!-- Project execution -->
-			<Agent Name="Automated Test Agent $(Project) $(Platform)" Type="EditorPerf$(Platform)">
+				<!-- Project execution -->
 				<Property Name="OutputName" Value="$(Project)$(Platform)EditorPerf"/>
 				<Property Name="RunNodeName" Value="Run $(OutputName)"/>
 				<Property Name="RunNodeRequires" Value="Compile Prerequisites $(Platform);Compile $(Project) Editor $(Platform)"/>
@@ -171,6 +169,11 @@
 				<Property Name="BaselineNodeRequires" Value=""/>
 				<Property Name="BaselineNodeRequires" Value="$(RunNodeName)" If="$(RunAndBaseline)"/>
 
+				<!-- Build the Editor for this Project -->
+				<Node Name="Compile $(Project) Editor $(Platform)" Requires="Compile Prerequisites $(Platform)">
+					<Compile Target="$(EditorTarget)" Platform="$(Platform)" Configuration="Development" Arguments="$(EditorCompileArgs)"/>
+				</Node>
+
 				<!-- Run the Tests -->
 				<Node Name="$(RunNodeName)" Requires="$(RunNodeRequires)">
 					<ForEach Name="Phase" Values="$(Phases)">
@@ -188,9 +191,10 @@
 				<Property Name="BaselineGraphNodes" Value="$(BaselineGraphNodes);$(BaselineNodeName)"/>
 
 				<Label Category="Test" Name="$(Project) EditorPerf $(Platform)" Requires="$(RunNodeName)"/>
-			</Agent>
 
-		</ForEach> <!-- Projects -->
+			</ForEach> <!-- Projects -->
+
+		</Agent> <!-- Automated Test Agent $(Platform) -->
 
 	</ForEach> <!-- Platforms -->
 
diff --git a/Engine/Plugins/Tests/EditorTests/Source/EditorTests/Private/UnrealEd/EditorIterationTests.cpp b/Engine/Plugins/Tests/EditorTests/Source/EditorTests/Private/UnrealEd/EditorIterationTests.cpp
--- a/Engine/Plugins/Tests/EditorTests/Source/EditorTests/Private/UnrealEd/EditorIterationTests.cpp
+++ b/Engine/Plugins/Tests/EditorTests/Source/EditorTests/Private/UnrealEd/EditorIterationTests.cpp
@@ -19,12 +19,12 @@
  * Test to open the sub editor windows for a specified list of assets.
  * This list can be setup in the Editor Preferences window within the editor or the DefaultEngine.ini file for that particular project.
 */
-IMPLEMENT_CUSTOM_COMPLEX_AUTOMATION_TEST(FIterationOpenAssets, FAutomationTestBase, "Project.Iteration.OpenAssets", EAutomationTestFlags::EditorContext | EAutomationTestFlags::ProductFilter);
+IMPLEMENT_CUSTOM_COMPLEX_AUTOMATION_TEST(FIterationOpenAssets, FAutomationTestBase, "Project.Iteration.PIE", EAutomationTestFlags::EditorContext | EAutomationTestFlags::ProductFilter);
 
 void FIterationOpenAssets::GetTests(TArray<FString>& OutBeautifiedNames, TArray<FString>& OutTestCommands) const
 {
 	const UAutomationTestSettings* AutomationTestSettings = GetDefault<UAutomationTestSettings>();
-	for ( FSoftObjectPath AssetRef : AutomationTestSettings->AssetsToOpen )
+	for ( FSoftObjectPath AssetRef : AutomationTestSettings->MapsToPIETest)
 	{
 		OutBeautifiedNames.Add(AssetRef.GetAssetName());
 		OutTestCommands.Add(AssetRef.GetLongPackageName());
