Fix Agent setup to parallelize across projects

#rnd

diff --git a/Engine/Build/EditorPerf.xml b/Engine/Build/EditorPerf.xml
--- a/Engine/Build/EditorPerf.xml
+++ b/Engine/Build/EditorPerf.xml
@@ -143,8 +143,7 @@
 	<!-- ForEach Platform to run, generate graph nodes -->
 	<ForEach Name="Platform" Values="$(Platforms)">
 
-		<!-- We do everything on a single resource because it's a beast and we don't want to wait or use caching/distribution because it will pollute perf numbers -->
-		<Agent Name="Automated Test Agent $(Platform)" Type="EditorPerf$(Platform)">
+		<Agent Name="Prerequisites Agent $(Platform)" Type="$($(Platform)CompileAgentType)">
 
 			<!-- Build Platform-global prerequisite targets -->
 			<Node Name="Compile Prerequisites $(Platform)">
@@ -154,20 +153,24 @@
 				<Compile Target="CrashReportClientEditor" Platform="$(Platform)" Configuration="Shipping"/>
 			</Node>
 
-			<!-- ForEach Project to run, generate the graph nodes -->
-			<ForEach Name="Project" Values="$(Projects)">
+		</Agent>
 
-				<!-- We stuff the Editor target name in [] following the project name in case it's not $(Project)Editor -->
-				<Property Name="EditorTarget" Value="$(Project)Editor"/>
-				<Regex Pattern="(.*)\[(.*)\]" Capture="Project;EditorTarget" Input="$(Project)" Optional="true"/>
+		<!-- ForEach Project to run, generate the graph nodes -->
+		<ForEach Name="Project" Values="$(Projects)">
 
-				<!-- Project execution -->
-				<Property Name="OutputName" Value="$(Project)$(Platform)EditorPerf"/>
-				<Property Name="RunNodeName" Value="Run $(OutputName)"/>
-				<Property Name="RunNodeRequires" Value="Compile Prerequisites $(Platform);Compile $(Project) Editor $(Platform)"/>
-				<Property Name="BaselineNodeName" Value="Baseline $(OutputName)"/>
-				<Property Name="BaselineNodeRequires" Value=""/>
-				<Property Name="BaselineNodeRequires" Value="$(RunNodeName)" If="$(RunAndBaseline)"/>
+			<!-- We stuff the Editor target name in [] following the project name in case it's not $(Project)Editor -->
+			<Property Name="EditorTarget" Value="$(Project)Editor"/>
+			<Regex Pattern="(.*)\[(.*)\]" Capture="Project;EditorTarget" Input="$(Project)" Optional="true"/>
+
+			<!-- Project execution -->
+			<Property Name="OutputName" Value="$(Project)$(Platform)EditorPerf"/>
+			<Property Name="RunNodeName" Value="Run $(OutputName)"/>
+			<Property Name="RunNodeRequires" Value="Compile Prerequisites $(Platform);Compile $(Project) Editor $(Platform)"/>
+			<Property Name="BaselineNodeName" Value="Baseline $(OutputName)"/>
+			<Property Name="BaselineNodeRequires" Value=""/>
+			<Property Name="BaselineNodeRequires" Value="$(RunNodeName)" If="$(RunAndBaseline)"/>
+
+			<Agent Name="Automated Test Agent $(Platform) $(Project)" Type="EditorPerf$(Platform)">
 
 				<!-- Build the Editor for this Project -->
 				<Node Name="Compile $(Project) Editor $(Platform)" Requires="Compile Prerequisites $(Platform)">
@@ -191,10 +194,10 @@
 				<Property Name="BaselineGraphNodes" Value="$(BaselineGraphNodes);$(BaselineNodeName)"/>
 
 				<Label Category="Test" Name="$(Project) EditorPerf $(Platform)" Requires="$(RunNodeName)"/>
+			
+			</Agent>
 
-			</ForEach> <!-- Projects -->
-
-		</Agent> <!-- Automated Test Agent $(Platform) -->
+		</ForEach> <!-- Projects -->
 
 	</ForEach> <!-- Platforms -->
 
