Fix InstallBuild for Linux at the final copy files stage.

diff --git a/Engine/Build/InstalledEngineFilters.xml b/Engine/Build/InstalledEngineFilters.xml
--- a/Engine/Build/InstalledEngineFilters.xml
+++ b/Engine/Build/InstalledEngineFilters.xml
@@ -362,7 +362,7 @@
 		Engine/Extras/GDBPrinters/UEPrinters.py
 
 		<!-- Bundled libcxx -->
-		Engine/Source/ThirdParty/Linux/LibCxx/...
+		Engine/Source/ThirdParty/Unix/LibCxx/...
 	</Property>
 	<Property Name="CopyEditorExceptionsLinux" Value="$(CopyEditorExceptions)">
 		FeaturePacks/TP_VirtualRealityBP.upack
@@ -495,9 +495,9 @@
 	<!-- Linux -->
 	<Property Name="CopyLinuxFilter">
 		Engine/Binaries/ThirdParty/OpenAL/...
-		Engine/Binaries/ThirdParty/ICU/icu4c-53_1/Linux/x86_64-unknown-linux-gnu/*.so
-		Engine/Source/ThirdParty/Linux/LibCxx/...
-		Engine/Source/ThirdParty/Linux/HaveLinuxDependencies
+		Engine/Binaries/ThirdParty/ICU/icu4c-53_1/Unix/x86_64-unknown-linux-gnu/*.so
+		Engine/Source/ThirdParty/Unix/LibCxx/...
+		Engine/Source/ThirdParty/Unix/HaveLinuxDependencies
 	</Property>
 
 	<Property Name="CopyLinuxFilterWin64" Value="$(CopyLinuxFilter)">
