Vivox is removed from UE5 and in a plugin now, no longer need to copy these files. One is missing

diff --git a/Engine/Build/InstalledEngineFilters.xml b/Engine/Build/InstalledEngineFilters.xml
--- a/Engine/Build/InstalledEngineFilters.xml
+++ b/Engine/Build/InstalledEngineFilters.xml
@@ -432,8 +432,6 @@
 		Engine/Source/ThirdParty/Oculus/OculusOpenXRLoader/OculusOpenXRLoader/Lib/*/*.so
 		Engine/Source/ThirdParty/ResonanceAudioApi/...
 		Engine/Source/ThirdParty/TangoSDK/...
-		Engine/Source/ThirdParty/Vivox/vivox-sdk/Lib/Android/sdk-release.aar
-		Engine/Source/ThirdParty/Vivox/vivox-sdk/VivoxCoreSDK_UPL.xml
 	</Property>
 	
 	<Property Name="CopyAndroidExceptions">
