Fix some missing build root path for Linux

diff --git a/Engine/Build/InstalledEngineBuild.xml b/Engine/Build/InstalledEngineBuild.xml
--- a/Engine/Build/InstalledEngineBuild.xml
+++ b/Engine/Build/InstalledEngineBuild.xml
@@ -517,8 +517,8 @@
 			<Sign Files="#Sign Binaries" If="'$(SignExecutables)' == true"/>
 		</Node>
 		<Node Name="Build Tools Linux" Requires="Compile UnrealHeaderTool Win64">
-			<Compile Target="CrashReportClient" Platform="Linux" Configuration="Shipping"/>
-			<Compile Target="CrashReportClientEditor" Platform="Linux" Configuration="Shipping"/>
+			<Compile Target="CrashReportClient" Platform="Linux" Configuration="Shipping" Tag="#Build Tools Linux" Arguments="$(CrashReporterCompileArgs)"/>
+			<Compile Target="CrashReportClientEditor" Platform="Linux" Configuration="Shipping" Tag="#Build Tools Linux" Arguments="$(CrashReporterCompileArgs)"/>
 			<Compile Target="ShaderCompileWorker" Configuration="Development" Platform="Linux" Tag="#Build Tools Linux"/>
 			<Compile Target="UnrealInsights" Configuration="Development" Platform="Linux" Tag="#Build Tools Linux"/>
 			<Compile Target="UnrealLightmass" Configuration="Development" Platform="Linux" Tag="#Build Tools Linux"/>
@@ -594,8 +594,8 @@
 
 	<Agent Name="Tools Group LinuxNative" Type="Linux_Licensee">
 		<Node Name="Build Tools LinuxNative" Requires="Compile UnrealHeaderTool Linux">
-			<Compile Target="CrashReportClient" Platform="Linux" Configuration="Shipping"/>
-			<Compile Target="CrashReportClientEditor" Platform="Linux" Configuration="Shipping"/>
+			<Compile Target="CrashReportClient" Platform="Linux" Configuration="Shipping" Tag="#Build Tools LinuxNative" Arguments="$(CrashReporterCompileArgs)"/>
+			<Compile Target="CrashReportClientEditor" Platform="Linux" Configuration="Shipping"  Tag="#Build Tools LinuxNative" Arguments="$(CrashReporterCompileArgs)"/>
 			<Compile Target="ShaderCompileWorker" Configuration="Development" Platform="Linux" Tag="#Build Tools LinuxNative"/>
 			<Compile Target="UnrealInsights" Configuration="Development" Platform="Linux" Tag="#Build Tools LinuxNative"/>
 			<Compile Target="UnrealLightmass" Configuration="Development" Platform="Linux" Tag="#Build Tools LinuxNative"/>
diff --git a/Engine/Source/Programs/AutomationTool/AutomationUtils/CommandUtils.cs b/Engine/Source/Programs/AutomationTool/AutomationUtils/CommandUtils.cs
--- a/Engine/Source/Programs/AutomationTool/AutomationUtils/CommandUtils.cs
+++ b/Engine/Source/Programs/AutomationTool/AutomationUtils/CommandUtils.cs
@@ -2037,7 +2037,7 @@ namespace AutomationTool
         /// <summary>
         /// "P:\Builds" or "/Volumes/Builds". Root Folder for all build storage.
         /// </summary>
-        /// <returns>"P:\Builds" or "/Volumes/Builds" unless overridden by -UseLocalBuildStorage from the commandline, where is uses Engine\Saved\LocalBuilds\.</returns>
+        /// <returns>"P:\Builds" or "/Volumes/Builds" or "/mnt/Builds" unless overridden by -UseLocalBuildStorage from the commandline, where is uses Engine\Saved\LocalBuilds\.</returns>
         public static string RootBuildStorageDirectory()
         {
             if (string.IsNullOrEmpty(CachedRootBuildStorageDirectory))
@@ -2050,7 +2050,18 @@ namespace AutomationTool
                 }
                 else
                 {
-                    CachedRootBuildStorageDirectory = RuntimePlatform.IsWindows ? CombinePaths("P:", "Builds") : "/Volumes/Builds";
+					if (RuntimePlatform.IsMac)
+					{
+						CachedRootBuildStorageDirectory = CombinePaths("Volumes", "Builds");
+					}
+					else if (RuntimePlatform.IsLinux)
+					{
+						CachedRootBuildStorageDirectory = CombinePaths("mnt", "Builds");
+					}
+					else
+					{
+						CachedRootBuildStorageDirectory = CombinePaths("P:", "Builds");
+					}
                 }
             }
             return CachedRootBuildStorageDirectory;
