EditorPerf Attempt fixes for Horde:

Move Editor Compiles back onto WIN-UE5 agents
Set Iteratiions back to 1

diff --git a/Engine/Build/EditorPerf.xml b/Engine/Build/EditorPerf.xml
--- a/Engine/Build/EditorPerf.xml
+++ b/Engine/Build/EditorPerf.xml
@@ -62,7 +62,7 @@
 	<Option Name="WarmArgs" DefaultValue="-ddc=Warm" Description="Arguments to unreal for the warm phase"/>
 
 	<!-- Iterations to run and average together -->
-	<Option Name="Iterations" DefaultValue="1;2;3" Description="The number of iterations to run"/>
+	<Option Name="Iterations" DefaultValue="1" Description="The number of iterations to run"/>
 
 	<!-- Gauntlet arguments for all runs -->
 	<Property Name="GauntletArgs" Value="-build=editor -UseEditor -verbose -tempdir=&quot;$(RootDir)/Tests&quot;"/>
@@ -143,7 +143,7 @@
 	<!-- ForEach Platform to run, generate graph nodes -->
 	<ForEach Name="Platform" Values="$(Platforms)">
 
-		<Agent Name="Prerequisites Agent $(Platform)" Type="$($(Platform)CompileAgentType)">
+		<Agent Name="Compile Prerequisites Agent $(Platform)" Type="$($(Platform)CompileAgentType)">
 
 			<!-- Build Platform-global prerequisite targets -->
 			<Node Name="Compile Prerequisites $(Platform)">
@@ -170,13 +170,17 @@
 			<Property Name="BaselineNodeRequires" Value=""/>
 			<Property Name="BaselineNodeRequires" Value="$(RunNodeName)" If="$(RunAndBaseline)"/>
 
-			<Agent Name="Automated Test Agent $(Platform) $(Project)" Type="EditorPerf$(Platform)">
+			<Agent Name="Compile Editor Agent $(Platform) $(Project)" Type="$($(Platform)CompileAgentType)">
 
 				<!-- Build the Editor for this Project -->
 				<Node Name="Compile $(Project) Editor $(Platform)" Requires="Compile Prerequisites $(Platform)">
 					<Compile Target="$(EditorTarget)" Platform="$(Platform)" Configuration="Development" Arguments="$(EditorCompileArgs)"/>
 				</Node>
 
+			</Agent>
+
+			<Agent Name="EditorPerf Agent $(Platform) $(Project)" Type="EditorPerf$(Platform)">
+
 				<!-- Run the Tests -->
 				<Node Name="$(RunNodeName)" Requires="$(RunNodeRequires)">
 					<ForEach Name="Phase" Values="$(Phases)">
