EditorPerf

Update for project-specific UHT targets

#rnx

diff --git a/Engine/Build/EditorPerf.xml b/Engine/Build/EditorPerf.xml
--- a/Engine/Build/EditorPerf.xml
+++ b/Engine/Build/EditorPerf.xml
@@ -143,38 +143,24 @@
 	<!-- ForEach Platform to run, generate graph nodes -->
 	<ForEach Name="Platform" Values="$(Platforms)">
 
-		<Agent Name="Compile Prerequisites Agent $(Platform)" Type="$($(Platform)CompileAgentType)">
-
-			<!-- Build Platform-global prerequisite targets -->
-			<Node Name="Compile Prerequisites $(Platform)">
-				<Compile Target="UnrealHeaderTool" Platform="$(Platform)" Configuration="Development"/>
-				<Compile Target="ShaderCompileWorker" Platform="$(Platform)" Configuration="Development"/>
-				<Compile Target="CrashReportClient" Platform="$(Platform)" Configuration="Shipping"/>
-				<Compile Target="CrashReportClientEditor" Platform="$(Platform)" Configuration="Shipping"/>
-			</Node>
-
-		</Agent>
-
 		<!-- ForEach Project to run, generate the graph nodes -->
 		<ForEach Name="Project" Values="$(Projects)">
 
-			<!-- We stuff the Editor target name in [] following the project name in case it's not $(Project)Editor -->
+			<!-- We parse the Editor target name from within [] in case it's not $(Project)Editor -->
 			<Property Name="EditorTarget" Value="$(Project)Editor"/>
 			<Regex Pattern="(.*)\[(.*)\]" Capture="Project;EditorTarget" Input="$(Project)" Optional="true"/>
 
-			<!-- Project execution -->
+			<!-- Use as base string for directories and key for node names -->
 			<Property Name="OutputName" Value="$(Project)$(Platform)EditorPerf"/>
-			<Property Name="RunNodeName" Value="Run $(OutputName)"/>
-			<Property Name="RunNodeRequires" Value="Compile Prerequisites $(Platform);Compile $(Project) Editor $(Platform)"/>
-			<Property Name="BaselineNodeName" Value="Baseline $(OutputName)"/>
-			<Property Name="BaselineNodeRequires" Value=""/>
-			<Property Name="BaselineNodeRequires" Value="$(RunNodeName)" If="$(RunAndBaseline)"/>
 
 			<Agent Name="Compile Editor Agent $(Platform) $(Project)" Type="$($(Platform)CompileAgentType)">
 
-				<!-- Build the Editor for this Project -->
-				<Node Name="Compile $(Project) Editor $(Platform)" Requires="Compile Prerequisites $(Platform)">
-					<Compile Target="$(EditorTarget)" Platform="$(Platform)" Configuration="Development" Arguments="$(EditorCompileArgs)"/>
+				<Node Name="Compile $(Project) $(Platform)">
+					<Compile Target="UnrealHeaderTool" Platform="$(Platform)" Configuration="Development"/>
+					<Compile Target="UnrealHeaderTool" Platform="$(Platform)" Project="$(Project)" Configuration="Development"/>
+					<Compile Target="ShaderCompileWorker" Platform="$(Platform)" Configuration="Development"/>
+					<Compile Target="CrashReportClientEditor" Platform="$(Platform)" Project="$(Project)" Configuration="Shipping"/>
+					<Compile Target="$(EditorTarget)" Platform="$(Platform)" Project="$(Project)" Configuration="Development" Arguments="$(EditorCompileArgs)"/>
 				</Node>
 
 			</Agent>
@@ -182,6 +168,8 @@
 			<Agent Name="EditorPerf Agent $(Platform) $(Project)" Type="EditorPerf$(Platform)">
 
 				<!-- Run the Tests -->
+				<Property Name="RunNodeName" Value="Run $(OutputName)"/>
+				<Property Name="RunNodeRequires" Value="Compile $(Project) $(Platform)"/>
 				<Node Name="$(RunNodeName)" Requires="$(RunNodeRequires)">
 					<ForEach Name="Phase" Values="$(Phases)">
 						<Expand Name="RunTest" OutputName="$(OutputName)" Platform="$(Platform)" Project="$(Project)" NodeName="$(RunNodeName)" Phase="$(Phase)" Args="$($(Phase)Args)"/>
@@ -190,6 +178,9 @@
 				<Property Name="RunGraphNodes" Value="$(RunGraphNodes);$(RunNodeName)"/>
 
 				<!-- Baseline the Results -->
+				<Property Name="BaselineNodeName" Value="Baseline $(OutputName)"/>
+				<Property Name="BaselineNodeRequires" Value=""/>
+				<Property Name="BaselineNodeRequires" Value="$(RunNodeName)" If="$(RunAndBaseline)"/>
 				<Node Name="$(BaselineNodeName)" Requires="$(BaselineNodeRequires)">
 					<ForEach Name="Phase" Values="$(Phases)">
 						<Expand Name="BaselineResult" Platform="$(Platform)" Project="$(Project)" Phase="$(Phase)"/>
