Clear the compatible changelist number in //UE5/Main.

diff --git a/Engine/Build/Build.version b/Engine/Build/Build.version
--- a/Engine/Build/Build.version
+++ b/Engine/Build/Build.version
@@ -3,7 +3,7 @@
 	"MinorVersion": 1,
 	"PatchVersion": 0,
 	"Changelist": 0,
-	"CompatibleChangelist": 16433597,
+	"CompatibleChangelist": 0,
 	"IsLicenseeVersion": 0,
 	"IsPromotedBuild": 0,
 	"BranchName": "UE5"
