Fixing Build DDC Mac failing to launch UBT

UBT wasn't being compiled as part of the Node like the other platforms, restored UBT compiling

#rnx

diff --git a/Engine/Build/InstalledEngineBuild.xml b/Engine/Build/InstalledEngineBuild.xml
--- a/Engine/Build/InstalledEngineBuild.xml
+++ b/Engine/Build/InstalledEngineBuild.xml
@@ -757,6 +757,7 @@
 			<!-- Delete all the old files -->
 			<Delete Files="$(DDCDir)/..."/>
 			<!-- Build up a list of files needed to build DDC -->
+			<CsCompile Project="Engine/Source/Programs/UnrealBuildTool/UnrealBuildTool.csproj" Configuration="Development" Platform="AnyCPU" Tag="#ToCopy" TagReferences="#ToCopy" EnumerateOnly="true"/>
 			<Tag Files="#Compile UnrealHeaderTool Mac;#UnrealEditor Mac;#Build Tools Mac" With="#ToCopy"/>
 			<Tag Files="#ToCopy" Filter="*.target" With="#TargetReceipts"/>
 			<TagReceipt Files="#TargetReceipts" RuntimeDependencies="true" With="#ToCopy"/>
