Add BuildPatchTool to Linux Tools

diff --git a/Engine/Build/InstalledEngineBuild.xml b/Engine/Build/InstalledEngineBuild.xml
--- a/Engine/Build/InstalledEngineBuild.xml
+++ b/Engine/Build/InstalledEngineBuild.xml
@@ -526,6 +526,7 @@
 			<Compile Target="UnrealLightmass" Configuration="Development" Platform="Linux" Tag="#Build Tools Linux"/>
 			<Compile Target="UnrealPak" Configuration="Development" Platform="Linux" Tag="#Build Tools Linux"/>
 			<Compile Target="UnrealMultiUserServer" Configuration="Development" Platform="Linux" Tag="#Build Tools Linux"/>
+			<Compile Target="BuildPatchTool" Configuration="Shipping" Platform="Linux" Tag="#Build Tools Linux"/>
 			<Compile Target="UnrealRecoverySvc" Configuration="Development" Platform="Linux" Tag="#Build Tools Linux"/>
 			<Compile Target="SwitchboardListener" Configuration="Development" Platform="Linux" Tag="#Build Tools Linux"/>
 		</Node>
@@ -603,6 +604,7 @@
 			<Compile Target="UnrealLightmass" Configuration="Development" Platform="Linux" Tag="#Build Tools LinuxNative"/>
 			<Compile Target="UnrealPak" Configuration="Development" Platform="Linux" Tag="#Build Tools LinuxNative"/>
 			<Compile Target="UnrealMultiUserServer" Configuration="Development" Platform="Linux" Tag="#Build Tools LinuxNative"/>
+			<Compile Target="BuildPatchTool" Configuration="Shipping" Platform="Linux" Tag="#Build Tools LinuxNative"/>
 			<Compile Target="UnrealRecoverySvc" Configuration="Development" Platform="Linux" Tag="#Build Tools LinuxNative"/>
 			<Compile Target="SwitchboardListener" Configuration="Development" Platform="Linux" Tag="#Build Tools LinuxNative"/>
 		</Node>
