Add options for label categories, which can vary on different Horde servers
Move properties out of agent blocks (they only are necessary for Node)

#rnd
#presubmit 615b51abe3b24d00016b6954

diff --git a/Engine/Build/EditorPerf.xml b/Engine/Build/EditorPerf.xml
--- a/Engine/Build/EditorPerf.xml
+++ b/Engine/Build/EditorPerf.xml
@@ -11,6 +11,9 @@
 	<Property Name="PreflightSuffix" Value=""/>
 	<Property Name="PreflightSuffix" Value="-PF-$(PreflightChange)-$(UE_HORDE_JOBID)" If="$(IsPreflight)"/>
 	<Property Name="BuildName" Value="$(EscapedBranch)-CL-$(Change)$(PreflightSuffix)"/>
+	<Option Name="CompileCategory" DefaultValue="Editors" Description="Category string for Compile Label"/>
+	<Option Name="RunCategory" DefaultValue="Test" Description="Category string for Run Label"/>
+	<Option Name="BaselineCategory" DefaultValue="Test" Description="Category string for Baseline Label"/>
 
 	<!-- Projects setup-->
 	<Option Name="Projects" DefaultValue="false" Restrict=".+" Description="The projects to test"/>
@@ -153,44 +156,45 @@
 			<!-- Use as base string for directories and key for node names -->
 			<Property Name="OutputName" Value="$(Project)$(Platform)EditorPerf"/>
 
+			<!-- Compile the Editor -->
+			<Property Name="CompileNodeName" Value="Compile $(Project) $(Platform)"/>
 			<Agent Name="Compile Editor Agent $(Platform) $(Project)" Type="$($(Platform)CompileAgentType)">
-
-				<Node Name="Compile $(Project) $(Platform)">
+				<Node Name="$(CompileNodeName)">
 					<Compile Target="UnrealHeaderTool" Platform="$(Platform)" Configuration="Development"/>
 					<Compile Target="UnrealHeaderTool" Platform="$(Platform)" Project="$(Project)" Configuration="Development"/>
 					<Compile Target="ShaderCompileWorker" Platform="$(Platform)" Configuration="Development"/>
 					<Compile Target="CrashReportClientEditor" Platform="$(Platform)" Project="$(Project)" Configuration="Shipping"/>
 					<Compile Target="$(EditorTarget)" Platform="$(Platform)" Project="$(Project)" Configuration="Development" Arguments="$(EditorCompileArgs)"/>
 				</Node>
-
 			</Agent>
+			<Label Category="$(CompileCategory)" Name="EP Compile $(Project) $(Platform)" Requires="$(CompileNodeName)"/>
 
+			<!-- Run the Tests -->
+			<Property Name="RunNodeName" Value="Run $(OutputName)"/>
+			<Property Name="RunNodeRequires" Value="$(CompileNodeName)"/>
+			<Property Name="RunGraphNodes" Value="$(RunGraphNodes);$(RunNodeName)"/>
 			<Agent Name="EditorPerf Agent $(Platform) $(Project)" Type="EditorPerf$(Platform)">
-
-				<!-- Run the Tests -->
-				<Property Name="RunNodeName" Value="Run $(OutputName)"/>
-				<Property Name="RunNodeRequires" Value="Compile $(Project) $(Platform)"/>
 				<Node Name="$(RunNodeName)" Requires="$(RunNodeRequires)">
 					<ForEach Name="Phase" Values="$(Phases)">
 						<Expand Name="RunTest" OutputName="$(OutputName)" Platform="$(Platform)" Project="$(Project)" NodeName="$(RunNodeName)" Phase="$(Phase)" Args="$($(Phase)Args)"/>
 					</ForEach>
 				</Node>
-				<Property Name="RunGraphNodes" Value="$(RunGraphNodes);$(RunNodeName)"/>
+			</Agent>
+			<Label Category="$(RunCategory)" Name="EP Run $(Project) $(Platform)" Requires="$(RunNodeName)"/>
 
-				<!-- Baseline the Results -->
-				<Property Name="BaselineNodeName" Value="Baseline $(OutputName)"/>
-				<Property Name="BaselineNodeRequires" Value=""/>
-				<Property Name="BaselineNodeRequires" Value="$(RunNodeName)" If="$(RunAndBaseline)"/>
+			<!-- Baseline the Results -->
+			<Property Name="BaselineNodeName" Value="Baseline $(OutputName)"/>
+			<Property Name="BaselineNodeRequires" Value=""/>
+			<Property Name="BaselineNodeRequires" Value="$(RunNodeName)" If="$(RunAndBaseline)"/>
+			<Property Name="BaselineGraphNodes" Value="$(BaselineGraphNodes);$(BaselineNodeName)"/>
+			<Agent Name="EditorPerf Agent $(Platform) $(Project)" Type="EditorPerf$(Platform)">
 				<Node Name="$(BaselineNodeName)" Requires="$(BaselineNodeRequires)">
 					<ForEach Name="Phase" Values="$(Phases)">
 						<Expand Name="BaselineResult" Platform="$(Platform)" Project="$(Project)" Phase="$(Phase)"/>
 					</ForEach>
 				</Node>
-				<Property Name="BaselineGraphNodes" Value="$(BaselineGraphNodes);$(BaselineNodeName)"/>
-
-				<Label Category="Test" Name="$(Project) EditorPerf $(Platform)" Requires="$(RunNodeName)"/>
-			
 			</Agent>
+			<Label Category="$(BaselineCategory)" Name="EP Baseline $(Project) $(Platform)" Requires="$(BaselineNodeName)"/>
 
 		</ForEach> <!-- Projects -->
 
