Add missing Android filter

#android

diff --git a/Engine/Build/InstalledEngineFilters.xml b/Engine/Build/InstalledEngineFilters.xml
--- a/Engine/Build/InstalledEngineFilters.xml
+++ b/Engine/Build/InstalledEngineFilters.xml
@@ -423,6 +423,7 @@
 		Engine/Source/ThirdParty/GoogleGameSDK/...
 		Engine/Source/ThirdParty/GoogleVR/...
 		Engine/Source/ThirdParty/GoogleInstantPreview/InstantPreview.apk
+		Engine/Source/ThirdParty/heapprofd/...
 		Engine/Source/ThirdParty/HWCPipe/...
 		Engine/Source/ThirdParty/Oculus/LibOVRAudio/LibOVRAudio/lib/*/*.so
 		Engine/Source/ThirdParty/Oculus/LibOVRAudio/*.tps
