Fix Linux installed builds not including Android even if the SDK was around

Fix include files being incorrect cased

diff --git a/Engine/Build/InstalledEngineBuild.xml b/Engine/Build/InstalledEngineBuild.xml
--- a/Engine/Build/InstalledEngineBuild.xml
+++ b/Engine/Build/InstalledEngineBuild.xml
@@ -731,6 +731,7 @@
 	<Property Name="DDCPlatformsLinux" Value="$(DDCPlatformsLinux)+Linux" If="'$(WithLinux)' == true"/>
 	<Do If="'$(HostPlatformDDCOnly)' == false">
 		<Property Name="DDCPlatformsLinux" Value="$(DDCPlatformsLinux)+LinuxAArch64" If="'$(WithLinuxAArch64)' == true"/>
+		<Property Name="DDCPlatformsLinux" Value="$(DDCPlatformsLinux)+Android_ATC" If="'$(WithAndroid)' == true"/>
 	</Do>
 
 	<Agent Name="DDC Group Linux" Type="Linux;Win64">
@@ -1097,7 +1098,7 @@
 		<!-- Optional Target Platform requirements -->
 		<Property Name="InstalledRequirements" Value="$(InstalledRequirements);#UnrealGame Linux;#UnrealGame Linux Unstripped;#UnrealGame Linux Stripped" If="'$(WithLinux)' == true"/>
 		<Property Name="InstalledRequirements" Value="$(InstalledRequirements);#UnrealGame LinuxAArch64;#UnrealGame LinuxAArch64 Unstripped;#UnrealGame LinuxAArch64 Stripped" If="'$(WithLinuxAArch64)' == true"/>
-		<Property Name="InstalledRequirements" Value="$(InstalledRequirements);#CompressedDDCLinux" If="'$(WithDDC)' == true And $(WithLinux)"/>
+		<Property Name="InstalledRequirements" Value="$(InstalledRequirements);#UE4Game Android;#UE4Game Android Unstripped;#UE4Game Android Stripped" If="'$(WithAndroid)' == true"/>		 		<Property Name="InstalledRequirements" Value="$(InstalledRequirements);#CompressedDDCLinux" If="'$(WithDDC)' == true And $(WithLinux)"/>
 
 		<!-- Feature packs -->
 		<Property Name="InstalledRequirements" Value="$(InstalledRequirements);Make Feature Packs"/>
@@ -1136,6 +1137,13 @@
 				<Property Name="CopyInstalledFilter" Value="$(CopyInstalledFilter);$(CopyLinuxFilter)"/>
 				<Property Name="CopyInstalledExceptions" Value="$(CopyInstalledExceptions);$(CopyLinuxExceptions)"/>
 			</Do>
+			<Do If="'$(WithAndroid)' == true">
+				<Tag Files="#UE4Game Android" Except="#UE4Game Android Unstripped" With="#Installed Linux"/>
+				<Tag Files="#UE4Game Android Stripped" With="#Saved Output"/>
+				<Property Name="Platforms" Value="$(Platforms)Android;"/>
+				<Property Name="CopyInstalledFilter" Value="$(CopyInstalledFilter);$(CopyAndroidFilter)"/>
+				<Property Name="CopyInstalledExceptions" Value="$(CopyInstalledExceptions);$(CopyAndroidExceptions)"/>
+			</Do>
 
 			<!-- Tag any dependencies from all previous build products -->
 			<Tag Files="#Installed Linux" Filter="*.target" With="#TargetReceipts"/>
diff --git a/Engine/Plugins/Media/ElectraUtil/Source/ElectraBase/Public/MediaVideoDecoderOutputAndroid.h b/Engine/Plugins/Media/ElectraUtil/Source/ElectraBase/Public/MediaVideoDecoderOutputAndroid.h
--- a/Engine/Plugins/Media/ElectraUtil/Source/ElectraBase/Public/MediaVideoDecoderOutputAndroid.h
+++ b/Engine/Plugins/Media/ElectraUtil/Source/ElectraBase/Public/MediaVideoDecoderOutputAndroid.h
@@ -5,7 +5,7 @@
 
 #include "Templates/RefCounting.h"
 #include "Containers/Array.h"
-#include "Pixelformat.h"
+#include "PixelFormat.h"
 
 class FVideoDecoderOutputAndroid : public FVideoDecoderOutput
 {
diff --git a/Engine/Plugins/Media/ElectraUtil/Source/ElectraSamples/Private/ElectraSamplesTextureManager.cpp b/Engine/Plugins/Media/ElectraUtil/Source/ElectraSamples/Private/ElectraSamplesTextureManager.cpp
--- a/Engine/Plugins/Media/ElectraUtil/Source/ElectraSamples/Private/ElectraSamplesTextureManager.cpp
+++ b/Engine/Plugins/Media/ElectraUtil/Source/ElectraSamples/Private/ElectraSamplesTextureManager.cpp
@@ -1,7 +1,7 @@
 // Copyright Epic Games, Inc. All Rights Reserved.
 
 #include "ElectraSamplesTextureManager.h"
-#include "misc/ScopeLock.h"
+#include "Misc/ScopeLock.h"
 
 FElectraSamplesTextureManager::FElectraSamplesTextureManager()
 {
