EditorPerf: Allocate a <Node> per Phase and Iteration to make Horde results easier to read

#rnx

diff --git a/Engine/Build/EditorPerf.xml b/Engine/Build/EditorPerf.xml
--- a/Engine/Build/EditorPerf.xml
+++ b/Engine/Build/EditorPerf.xml
@@ -65,7 +65,7 @@
 	<Option Name="WarmArgs" DefaultValue="-ddc=Warm" Description="Arguments to unreal for the warm phase"/>
 
 	<!-- Iterations to run and average together -->
-	<Option Name="Iterations" DefaultValue="1" Description="The number of iterations to run"/>
+	<Option Name="Iterations" DefaultValue="1" Description="The list of iteration numbers to run. For example '1;2;3' for three iterations"/>
 
 	<!-- Gauntlet arguments for all runs -->
 	<Property Name="GauntletArgs" Value="-build=editor -UseEditor -verbose -tempdir=&quot;$(RootDir)/Tests&quot;"/>
@@ -91,40 +91,34 @@
 	<Property Name="TelemetryConfig" Value="UETelemetry" If="$(IsBuildMachine) And '$(TelemetryConfig)' == ''"/>
 
 	<!-- The main test running macro logic that is expanded with platform and dependency node names -->
-	<Macro Name="RunTest" Arguments="OutputName;Platform;Project;NodeName;Phase;Args">
-		<Property Name="TestName" Value="EditorPerf$(Phase)"/>
-
-		<ForEach Name="Iteration" Values="$(Iterations)">
+	<Macro Name="RunTest" Arguments="OutputName;Platform;Project;NodeName;Phase;Iteration;Args">
 
-			<!-- Only ever do one iteration of incremental, because it will warm up the workspace -->
-			<Do If="'$(Phase)' != 'Incr' or '$(Iteration)' == '1'">
-				<Property Name="IterationName" Value="$(Project)$(Platform)$(TestName)$(Iteration)"/>
-				<Property Name="TraceFilePath" Value="$(TraceDir)/$(IterationName).utrace"/>
-				<Property Name="TelemetryFileName" Value="$(IterationName)Telemetry.csv"/>
-				<Property Name="TelemetryFilePath" Value="$(TraceDir)/$(TelemetryFileName)"/>
-				<Property Name="LogOutputParams" Value="-logdir=&quot;$(AutomationOutputDir)/$(OutputName)&quot;"/>
-				<Property Name="ReportParams" Value="-WriteTestResultsForHorde -HordeTestDataKey=&quot;$(NodeName) $(Phase)$(Iteration)$(PreflightSuffix)&quot;"/>
-				<Property Name="AllGauntletArgs" Value="-Project=$(Project) -Platform=$(Platform) $(GauntletArgs) $(ReportParams) $(LogOutputParams) -artifactname=&quot;$(Phase)$(Iteration)&quot;"/>
+		<Property Name="TestName" Value="EditorPerf$(Phase)"/>
+		<Property Name="IterationName" Value="$(Project)$(Platform)$(TestName)$(Iteration)"/>
+		<Property Name="TraceFilePath" Value="$(TraceDir)/$(IterationName).utrace"/>
+		<Property Name="TelemetryFileName" Value="$(IterationName)Telemetry.csv"/>
+		<Property Name="TelemetryFilePath" Value="$(TraceDir)/$(TelemetryFileName)"/>
+		<Property Name="LogOutputParams" Value="-logdir=&quot;$(AutomationOutputDir)/$(OutputName)&quot;"/>
+		<Property Name="ReportParams" Value="-WriteTestResultsForHorde -HordeTestDataKey=&quot;$(NodeName)$(PreflightSuffix)&quot;"/>
+		<Property Name="AllGauntletArgs" Value="-Project=$(Project) -Platform=$(Platform) $(GauntletArgs) $(ReportParams) $(LogOutputParams) -artifactname=&quot;$(Phase)$(Iteration)&quot;"/>
 
-				<!-- Do DDC cleanup outside the editor process, handles Cold cleanup without changing stats -->
-				<Commandlet Project="$(Project)" Name="DDCCleanup" Arguments="$(Args) $(CommandletArgs)"/>
+		<!-- Do DDC cleanup outside the editor process, handles Cold cleanup without changing stats -->
+		<Commandlet Project="$(Project)" Name="DDCCleanup" Arguments="$(Args) $(CommandletArgs)"/>
 
-				<!-- Run the Gauntlet EditorPerf Tests -->
-				<Command Name="RunUnreal" Arguments="$(AllGauntletArgs) $(ProfileArgs) $(Args) -tracefile=&quot;$(TraceFilePath)&quot;"/>
+		<!-- Run the Gauntlet EditorPerf Tests -->
+		<Command Name="RunUnreal" Arguments="$(AllGauntletArgs) $(ProfileArgs) $(Args) -tracefile=&quot;$(TraceFilePath)&quot;"/>
 
-				<!-- Translate the utrace file into csv data -->
-				<Commandlet Project="$(Project)" Name="SummarizeTrace" Arguments="-inputfile=&quot;$(TraceFilePath)&quot; -testname=&quot;$(TestName)&quot; $(SummarizeArgs) $(CommandletArgs)"/>
+		<!-- Translate the utrace file into csv data -->
+		<Commandlet Project="$(Project)" Name="SummarizeTrace" Arguments="-inputfile=&quot;$(TraceFilePath)&quot; -testname=&quot;$(TestName)&quot; $(SummarizeArgs) $(CommandletArgs)"/>
 
-				<!-- Save all files named for this run in the trace directory -->
-				<Copy From="$(TraceDir)/$(IterationName)*.*" To="$(AutomationOutputDir)"/>
+		<!-- Save all files named for this run in the trace directory -->
+		<Copy From="$(TraceDir)/$(IterationName)*.*" To="$(AutomationOutputDir)"/>
 
-				<Do If="'$(TelemetryConfig)' != ''">
-					<!-- Publish telemetry to the db -->
-					<Command Name="PublishUnrealAutomationTelemetry" Arguments="-CSVFile=&quot;$(TelemetryFilePath)&quot; -project=$(Project) -changelist=$(Change) -branch=$(Branch) -TelemetryConfig=$(TelemetryConfig)"/>
-				</Do>
-			</Do>
+		<Do If="'$(TelemetryConfig)' != ''">
+			<!-- Publish telemetry to the db -->
+			<Command Name="PublishUnrealAutomationTelemetry" Arguments="-CSVFile=&quot;$(TelemetryFilePath)&quot; -project=$(Project) -changelist=$(Change) -branch=$(Branch) -TelemetryConfig=$(TelemetryConfig)"/>
+		</Do>
 
-		</ForEach> <!-- Iterations -->
 	</Macro>
 
 	<!-- This will be set to the names of all nodes we want to execute based on specified platforms/options -->
@@ -173,17 +167,22 @@
 			<Label Category="$(CompileCategory)" Name="EP Compile $(Project) $(Platform)" Requires="$(CompileNodeName)"/>
 
 			<!-- Run the Tests -->
-			<Property Name="RunNodeName" Value="Run $(OutputName)"/>
 			<Property Name="RunNodeRequires" Value="$(CompileNodeName)"/>
-			<Property Name="RunGraphNodes" Value="$(RunGraphNodes);$(RunNodeName)"/>
 			<Agent Name="EditorPerf Agent $(Platform) $(Project)" Type="EditorPerf$(Platform)">
-				<Node Name="$(RunNodeName)" Requires="$(RunNodeRequires)">
-					<ForEach Name="Phase" Values="$(Phases)">
-						<Expand Name="RunTest" OutputName="$(OutputName)" Platform="$(Platform)" Project="$(Project)" NodeName="$(RunNodeName)" Phase="$(Phase)" Args="$($(Phase)Args)"/>
+				<ForEach Name="Phase" Values="$(Phases)">
+					<ForEach Name="Iteration" Values="$(Iterations)">
+						<!-- Only ever do one iteration of incremental, because it will warm up the workspace -->
+						<Do If="'$(Phase)' != 'Incr' or '$(Iteration)' == '1'">
+							<Property Name="RunNodeName" Value="Run $(OutputName) $(Phase)$(Iteration)"/>
+							<Property Name="RunGraphNodes" Value="$(RunGraphNodes);$(RunNodeName)"/>
+							<Node Name="$(RunNodeName)" Requires="$(RunNodeRequires)">
+								<Expand Name="RunTest" OutputName="$(OutputName)" Platform="$(Platform)" Project="$(Project)" NodeName="$(RunNodeName)" Phase="$(Phase)" Iteration="$(Iteration)" Args="$($(Phase)Args)"/>
+							</Node>
+						</Do>
 					</ForEach>
-				</Node>
+				</ForEach>
 			</Agent>
-			<Label Category="$(RunCategory)" Name="EP Run $(Project) $(Platform)" Requires="$(RunNodeName)"/>
+			<Label Category="$(RunCategory)" Name="EP Run $(Project) $(Platform)" Requires="$(RunGraphNodes)"/>
 
 			<!-- Baseline the Results -->
 			<Property Name="BaselineNodeName" Value="Baseline $(OutputName)"/>
