Create RunNode dependency chain to early out of the chain upon failure

#rnx

diff --git a/Engine/Build/EditorPerf.xml b/Engine/Build/EditorPerf.xml
--- a/Engine/Build/EditorPerf.xml
+++ b/Engine/Build/EditorPerf.xml
@@ -168,21 +168,23 @@
 
 			<!-- Run the Tests -->
 			<Property Name="RunNodeRequires" Value="$(CompileNodeName)"/>
+			<Property Name="RunNodeList" Value=""/> <!-- Each node in this the nested loops below, they need to depend on predecessors -->
 			<Agent Name="EditorPerf Agent $(Platform) $(Project)" Type="EditorPerf$(Platform)">
 				<ForEach Name="Phase" Values="$(Phases)">
 					<ForEach Name="Iteration" Values="$(Iterations)">
 						<!-- Only ever do one iteration of incremental, because it will warm up the workspace -->
 						<Do If="'$(Phase)' != 'Incr' or '$(Iteration)' == '1'">
 							<Property Name="RunNodeName" Value="Run $(OutputName) $(Phase)$(Iteration)"/>
-							<Property Name="RunGraphNodes" Value="$(RunGraphNodes);$(RunNodeName)"/>
-							<Node Name="$(RunNodeName)" Requires="$(RunNodeRequires)">
+							<Node Name="$(RunNodeName)" Requires="$(RunNodeRequires);$(RunNodeList)">
 								<Expand Name="RunTest" OutputName="$(OutputName)" Platform="$(Platform)" Project="$(Project)" NodeName="$(RunNodeName)" Phase="$(Phase)" Iteration="$(Iteration)" Args="$($(Phase)Args)"/>
 							</Node>
+							<Property Name="RunNodeList" Value="$(RunNodeList);$(RunNodeName)"/>
 						</Do>
 					</ForEach>
 				</ForEach>
 			</Agent>
-			<Label Category="$(RunCategory)" Name="EP Run $(Project) $(Platform)" Requires="$(RunGraphNodes)"/>
+			<Label Category="$(RunCategory)" Name="EP Run $(Project) $(Platform)" Requires="$(RunNodeRequires);$(RunNodeList)"/>
+			<Property Name="RunGraphNodes" Value="$(RunGraphNodes);$(RunNodeList)"/>
 
 			<!-- Baseline the Results -->
 			<Property Name="BaselineNodeName" Value="Baseline $(OutputName)"/>
