fix builtdirectory path having too many slashes

diff --git a/Engine/Build/InstalledEngineBuild.xml b/Engine/Build/InstalledEngineBuild.xml
--- a/Engine/Build/InstalledEngineBuild.xml
+++ b/Engine/Build/InstalledEngineBuild.xml
@@ -70,7 +70,7 @@
 	<Option Name="WithFullDebugInfo" Restrict="true|false" DefaultValue="false" Description="Generate full debug info for binary editor and packaged application builds"/>
 
 	<!-- Ouput directory for the build -->
-	<Option Name="BuiltDirectory" DefaultValue="$(RootDir)/LocalBuilds/Engine/" Description="Directory for outputting the built engine"/>
+	<Option Name="BuiltDirectory" DefaultValue="$(RootDir)/LocalBuilds/Engine" Description="Directory for outputting the built engine"/>
 
 	<!-- The local output directory -->
 	<Property Name="LocalInstalledDir" Value="$(BuiltDirectory)/Windows"/>
