BuildGraph: Add support for including scripts using wildcards.

diff --git a/Engine/Build/Graph/Examples/Properties.xml b/Engine/Build/Graph/Examples/Properties.xml
--- a/Engine/Build/Graph/Examples/Properties.xml
+++ b/Engine/Build/Graph/Examples/Properties.xml
@@ -82,10 +82,6 @@
 			</Do>
 			<Log Message="The property changed inside the conditional block is now set to $(LocalPropertyOutsideConditional)"/>
 
-			<Do If="Exists('Engine/Build/Graph/Examples/Properties.xml')">
-				<Log Message="This file does exist!"/>
-			</Do>
-
 			<Property Name="HaystackString" Value="One Two Three" />
 			<Property Name="HaystackItems" Value="One;Two;Three" />
 			<Property Name="Needle" Value="One" />
diff --git a/Engine/Build/InstalledEngineBuild.xml b/Engine/Build/InstalledEngineBuild.xml
--- a/Engine/Build/InstalledEngineBuild.xml
+++ b/Engine/Build/InstalledEngineBuild.xml
@@ -117,19 +117,10 @@
 		</Do>
 	</Macro>
 
-	<Property Name="ExtensionIncludes" Value=""/>
+	<!-- Include any scripts from platform extensions -->
 	<Property Name="ExtensionPlatforms" Value=""/>
-	<ForEach Name="Platform" Values="$(DataDrivenPlatforms)">
-		<Property Name="PlatformBuildScript" Value="Platforms/$(Platform)/Build/InstalledEngineBuild.xml"/>
-		<Do If="Exists('Engine/$(PlatformBuildScript)')">
-			<Property Name="ExtensionIncludes" Value="$(ExtensionIncludes);" If="'$(ExtensionIncludes)' != ''"/>
-			<Property Name="ExtensionIncludes" Value="$(ExtensionIncludes)../$(PlatformBuildScript)"/>
-			<Property Name="ExtensionPlatforms" Value="$(ExtensionPlatforms);" If="'$(ExtensionPlatforms)' != ''"/>
-			<Property Name="ExtensionPlatforms" Value="$(ExtensionPlatforms)$(Platform)"/>
-		</Do>
-	</ForEach>
-	<Include Script="$(ExtensionIncludes)"/>
-
+	<Include Script="../Platforms/*/Build/InstalledEngineBuild.xml"/>
+	
 	<!-- Include the script that lists files and patterns used when creating an installed build -->
 	<Include Script="InstalledEngineFilters.xml"/>
 
@@ -1218,7 +1209,7 @@
 
 			<!-- Copy the toolchain if it exists -->
 			<Property Name="LinuxToolchainDir" Value="$(ToolchainDir)/HostLinux"/>
-			<Copy Files="$(LinuxToolchainDir)/..." From="$(LinuxToolchainDir)" To="$(LocalInstalledDirLinux)/Engine/Extras/ThirdPartyNotUE/SDKs/HostLinux" If="Exists('$(LinuxToolchainDir)')" Tag="#Installed Build Linux Files"/>
+			<Copy Files="$(LinuxToolchainDir)/..." From="$(LinuxToolchainDir)" To="$(LocalInstalledDirLinux)/Engine/Extras/ThirdPartyNotUE/SDKs/HostLinux" Tag="#Installed Build Linux Files"/>
 
 			<!-- List all of the platforms that we only allow for content projects -->
 			<Property Name="ContentOnlyPlatforms" Value=""/>
diff --git a/Engine/Build/InstalledEngineFilters.xml b/Engine/Build/InstalledEngineFilters.xml
--- a/Engine/Build/InstalledEngineFilters.xml
+++ b/Engine/Build/InstalledEngineFilters.xml
@@ -200,8 +200,8 @@
 	<Property Name="CopyEditorFilter" Value="$(CopyEditorFilter);Engine/Source/UnrealClient.Target.cs" If="$(WithClient)"/>
 	<Property Name="CopyEditorFilter" Value="$(CopyEditorFilter);Engine/Source/UnrealServer.Target.cs" If="$(WithServer)"/>
 	<!-- Optional API docs -->
-	<Property Name="CopyEditorFilter" Value="$(CopyEditorFilter);Engine/Documentation/Builds/CppAPI-HTML.tgz" If="Exists('Engine/Documentation/Builds/CppAPI-HTML.tgz')"/>
-	<Property Name="CopyEditorFilter" Value="$(CopyEditorFilter);Engine/Documentation/Builds/BlueprintAPI-HTML.tgz" If="Exists('Engine/Documentation/Builds/BlueprintAPI-HTML.tgz')"/>
+	<Property Name="CopyEditorFilter" Value="$(CopyEditorFilter);Engine/Documentation/Builds/CppAPI-HTML.tgz"/>
+	<Property Name="CopyEditorFilter" Value="$(CopyEditorFilter);Engine/Documentation/Builds/BlueprintAPI-HTML.tgz"/>
 	<!-- Define Editor Exceptions -->
 	<Property Name="CopyEditorExceptions">
 	
diff --git a/Engine/Source/Programs/AutomationTool/BuildGraph/BuildGraph.cs b/Engine/Source/Programs/AutomationTool/BuildGraph/BuildGraph.cs
--- a/Engine/Source/Programs/AutomationTool/BuildGraph/BuildGraph.cs
+++ b/Engine/Source/Programs/AutomationTool/BuildGraph/BuildGraph.cs
@@ -141,6 +141,17 @@ namespace AutomationTool
 			}
 			return null;
 		}
+
+		/// <inheritdoc/>
+		public Task<string[]> FindAsync(string Pattern)
+		{
+			FileFilter Filter = new FileFilter();
+			Filter.AddRule(Pattern, FileFilterType.Include);
+
+			List<string> Files = Filter.ApplyToDirectory(Unreal.RootDirectory, true).ConvertAll(x => x.MakeRelativeTo(Unreal.RootDirectory).Replace('\\', '/'));
+			Files.Sort(StringComparer.OrdinalIgnoreCase);
+			return Task.FromResult(Files.ToArray());
+		}
 	}
 
 	/// <summary>
diff --git a/Engine/Source/Programs/Shared/EpicGames.BuildGraph/BgScriptReader.cs b/Engine/Source/Programs/Shared/EpicGames.BuildGraph/BgScriptReader.cs
--- a/Engine/Source/Programs/Shared/EpicGames.BuildGraph/BgScriptReader.cs
+++ b/Engine/Source/Programs/Shared/EpicGames.BuildGraph/BgScriptReader.cs
@@ -242,10 +242,16 @@ namespace EpicGames.BuildGraph
 		/// Tries to read a file from the given path
 		/// </summary>
 		/// <param name="Path">Path of the file to read</param>
-		/// <param name="Data">On success, contains the read data</param>
 		/// <returns></returns>
 		Task<byte[]?> ReadAsync(string Path);
 
+		/// <summary>
+		/// Finds files matching the given pattern
+		/// </summary>
+		/// <param name="Pattern"></param>
+		/// <returns></returns>
+		Task<string[]> FindAsync(string Pattern);
+
 		/// <summary>
 		/// Converts a path to its native form, for display to the user
 		/// </summary>
@@ -654,9 +660,24 @@ namespace EpicGames.BuildGraph
 		{
 			if (await EvaluateConditionAsync(Element))
 			{
+				HashSet<string> Files = new HashSet<string>();
 				foreach (string Script in ReadListAttribute(Element, "Script"))
 				{
-					await TryReadAsync(CombinePaths(Element.File, Script), Arguments, Logger);
+					string IncludePath = CombinePaths(Element.File, Script);
+					if (Regex.IsMatch(IncludePath, @"\*|\?|\.\.\."))
+					{
+						Files.UnionWith(await Context.FindAsync(IncludePath));
+					}
+					else
+					{
+						Files.Add(IncludePath);
+					}
+				}
+
+				foreach(string File in Files.OrderBy(x => x, StringComparer.OrdinalIgnoreCase))
+				{
+					Logger.LogDebug("Including file {File}", File);
+					await TryReadAsync(File, Arguments, Logger);
 				}
 			}
 		}
