Missing file.

diff --git a/Engine/Build/InstalledEngineBuild.xml b/Engine/Build/InstalledEngineBuild.xml
--- a/Engine/Build/InstalledEngineBuild.xml
+++ b/Engine/Build/InstalledEngineBuild.xml
@@ -117,6 +117,17 @@
 		</Do>
 	</Macro>
 
+	<!-- Create empty macros for platforms to extend -->
+	<Macro Name="Platform_FilterRestrictedFolders"/>
+	<Macro Name="Platform_FilterOverrides"/>
+	<Macro Name="Platform_Agent_Target_Win64_Licensee"/>
+	<Macro Name="Platform_AddInstalledRequirements"/>
+	<Macro Name="Platform_TagCsToolsRedistUATSource"/>
+	<Macro Name="Platform_Build_Tools_CS_Binaries"/>
+	<Macro Name="Platform_Append_DDCPlatformsWin64"/>
+	<Macro Name="Platform_TagDDCFiles"/>
+	<Macro Name="Platform_Tag_UnrealGame"/>
+
 	<!-- Include any scripts from platform extensions -->
 	<Property Name="ExtensionPlatforms" Value=""/>
 	<Include Script="../Platforms/*/Build/InstalledEngineBuild.xml"/>
@@ -349,9 +360,7 @@
 	</Agent>
 	<Aggregate Name="TargetPlatforms_Android" Label="Platforms/Android" Requires="Compile UnrealGame Android"/>
 
-	<ForEach Name="Platform" Values="$(ExtensionPlatforms)">
-		<Expand Name="Agent_Target_$(Platform)_Win64_Licensee"/>
-	</ForEach>
+	<Expand Name="Platform_Agent_Target_Win64_Licensee"/>
 
 	<Agent Name="Target HoloLens" Type="HoloLens">
 		<Node Name="Compile UnrealGame HoloLens" Requires="Compile UnrealHeaderTool Win64" Produces="#UnrealGame HoloLens">
@@ -473,9 +482,7 @@
 			<!-- Copy Source and referenced libraries to a new location with Confidential folders removed-->
 			<Tag Files="$(AutomationToolPaths)" Filter="*.csproj" Except="$(ConfidentialExceptions)"  With="#UAT Projects"/>
 			<Tag Files="Engine/Binaries/DotNET/...;Engine/Binaries/ThirdParty/Newtonsoft/...;Engine/Binaries/ThirdParty/IOS/...;Engine/Binaries/ThirdParty/VisualStudio/...;Engine/Source/Programs/...;Engine/Platforms/*/Source/Programs/...;Engine/Source/Editor/SwarmInterface/...;#UATReferences" Except="$(ConfidentialExceptions);AutomationTool.csproj.References" With="#RedistUATSource"/>
-			<ForEach Name="Platform" Values="$(ExtensionPlatforms)">
-				<Expand Name="$(Platform)_TagCsToolsRedistUATSource"/>
-			</ForEach>
+			<Expand Name="Platform_TagCsToolsRedistUATSource"/>
 			<Copy Files="#RedistUATSource" From="$(RootDir)" To="$(CsToolsDir)"/>
 
 			<Property Name="RedistAutomationToolPaths" Value=""/>
@@ -491,9 +498,7 @@
 				<CsCompile Project="$(CsToolsDir)/Engine/Source/Editor/SwarmInterface/DotNET/SwarmInterface.csproj" Configuration="Development" Platform="AnyCPU" Tag="#Build Tools CS Binaries;#Sign Binaries" TagReferences="#Build Tools CS Binaries"/>
 				<CsCompile Project="$(CsToolsDir)/Engine/Source/Programs/UnrealControls/UnrealControls.csproj" Configuration="Development" Platform="AnyCPU" Tag="#Build Tools CS Binaries;#Sign Binaries" TagReferences="#Build Tools CS Binaries"/>
 				<CsCompile Project="$(CsToolsDir)/Engine/Source/Programs/IOS/iPhonePackager/iPhonePackager.csproj" Configuration="Development" Platform="AnyCPU" Arguments="/verbosity:minimal /target:Rebuild" Tag="#Build Tools CS Binaries;#Sign Binaries" TagReferences="#Build Tools CS Binaries"/>
-				<ForEach Name="Platform" Values="$(ExtensionPlatforms)">
-					<Expand Name="$(Platform)_Build_Tools_CS_Binaries"/>
-				</ForEach>
+				<Expand Name="Platform_Build_Tools_CS_Binaries"/>
 			</Do>
 			
 			<!-- Tag AutomationTool and UnrealBuildTool folders recursively as NET Core dependencies are not currently handled by CsCompile -->
@@ -533,9 +538,7 @@
 			<!-- Copy Source and referenced libraries to a new location with Confidential folders removed-->
 			<Tag Files="$(AutomationToolPaths)" Filter="*.csproj" Except="$(ConfidentialExceptions)"  With="#UAT Projects Mac"/>
 			<Tag Files="Engine/Binaries/DotNET/...;Engine/Binaries/ThirdParty/Newtonsoft/...;Engine/Binaries/ThirdParty/IOS/...;Engine/Binaries/ThirdParty/VisualStudio/...;Engine/Source/Programs/...;Engine/Platforms/*/Source/Programs/...;Engine/Source/Editor/SwarmInterface/...;#UATReferencesMac" Except="$(ConfidentialExceptions);AutomationTool.csproj.References" With="#RedistUATSourceMac"/>
-			<ForEach Name="Platform" Values="$(ExtensionPlatforms)">
-				<Expand Name="$(Platform)_TagCsToolsRedistUATSource"/>
-			</ForEach>
+			<Expand Name="Platform_TagCsToolsRedistUATSource"/>
 			<Copy Files="#RedistUATSourceMac" From="$(RootDir)" To="$(CsToolsDir)"/>
 
 			<Property Name="RedistAutomationToolPaths" Value=""/>
@@ -607,9 +610,7 @@
 			<!-- Copy Source and referenced libraries to a new location with Confidential folders removed-->
 			<Tag Files="$(AutomationToolPaths)" Filter="*.csproj" Except="$(ConfidentialExceptions)"  With="#UAT Projects"/>
 			<Tag Files="Engine/Binaries/DotNET/...;Engine/Binaries/ThirdParty/Newtonsoft/...;Engine/Binaries/ThirdParty/IOS/...;Engine/Binaries/ThirdParty/VisualStudio/...;Engine/Source/Programs/...;Engine/Platforms/*/Source/Programs/...;Engine/Source/Editor/SwarmInterface/...;#UATReferences" Except="$(ConfidentialExceptions);AutomationTool.csproj.References" With="#RedistUATSource"/>
-			<ForEach Name="Platform" Values="$(ExtensionPlatforms)">
-				<Expand Name="$(Platform)_TagCsToolsRedistUATSource"/>
-			</ForEach>
+			<Expand Name="Platform_TagCsToolsRedistUATSource"/>
 			<Copy Files="#RedistUATSource" From="$(RootDir)" To="$(CsToolsDir)"/>
 
 			<Property Name="RedistAutomationToolPaths" Value=""/>
@@ -664,9 +665,7 @@
 		<Property Name="DDCPlatformsWin64" Value="$(DDCPlatformsWin64)+Linux" If="'$(WithLinux)' == true"/>
 		<Property Name="DDCPlatformsWin64" Value="$(DDCPlatformsWin64)+LinuxArm64" If="'$(WithLinuxArm64)' == true"/>
 		<Property Name="DDCPlatformsWin64" Value="$(DDCPlatformsWin64)+HoloLens" If="'$(WithHoloLens)' == true"/>
-		<ForEach Name="Platform" Values="$(ExtensionPlatforms)">
-			<Expand Name="$(Platform)_Append_DDCPlatformsWin64"/>
-		</ForEach>
+		<Expand Name="Platform_Append_DDCPlatformsWin64"/>
 	</Do>
 
 	<Macro Name="GenerateEngineDirList" Arguments="InList">
@@ -698,9 +697,9 @@
 		<Property Name="DDCCopyEngineShaders" Value="$(EngineDirList)"/>
 		<Expand Name="GenerateEngineDirList" InList="$(InstalledDDCEngineShadersExcept)"/>
 		<Property Name="DDCCopyEngineShadersExcept" Value="$(EngineDirList)"/>
-		<ForEach Name="Platform" Values="$(ExtensionPlatforms)">
-			<Expand Name="$(Platform)_TagDDCFiles"/>
-		</ForEach>
+		
+		<Expand Name="Platform_TagDDCFiles"/>
+		
 		<Tag Files="$(DDCCopyEngineContent)" Except="$(DDCCopyEngineContentExcept)" With="#ToCopy"/>
 		<Tag Files="$(DDCCopyEngineConfig)" Except="$(DDCCopyEngineConfigExcept)" With="#ToCopy"/>
 		<Tag Files="$(DDCCopyEnginePlugins)" Except="$(DDCCopyEnginePluginsExcept)" With="#ToCopy"/>
@@ -845,9 +844,7 @@
 		<Property Name="UnrealEditor" Value="#UnrealEditor Win64 WithExceptions" If="$(CompileDatasmithPlugins)" />
 		<Property Name="InstalledRequirements" Value="Compile UnrealHeaderTool Win64;$(UnrealEditor);#UnrealEditor Win64 Unstripped;#UnrealEditor Win64 Stripped;#UnrealEditor Win64 Unsigned;#UnrealEditor Win64 Signed;#Build Tools Win64;#Build Tools CS Binaries;#Rules Assemblies"/>
 		<!-- Optional Target Platform requirements -->
-		<ForEach Name="Platform" Values="$(ExtensionPlatforms)">
-			<Expand Name="$(Platform)_AddInstalledRequirements"/>
-		</ForEach>
+		<Expand Name="Platform_AddInstalledRequirements"/>
 		<Property Name="InstalledRequirements" Value="$(InstalledRequirements);#UnrealGame Win64;#UnrealGame Win64 Unstripped;#UnrealGame Win64 Stripped;#UnrealGame Win64 Unsigned;#UnrealGame Win64 Signed" If="'$(WithWin64)' == true"/>
 		<Property Name="InstalledRequirements" Value="$(InstalledRequirements);#UnrealGame Android;#UnrealGame Android Unstripped;#UnrealGame Android Stripped" If="'$(WithAndroid)' == true"/>
 		<Property Name="InstalledRequirements" Value="$(InstalledRequirements);#UnrealGame IOS;#UnrealGame IOS Unstripped;#UnrealGame IOS Stripped;#Compile UnrealHeaderTool Mac" If="'$(WithIOS)' == true"/>
@@ -923,9 +920,7 @@
 				<Property Name="CopyInstalledFilter" Value="$(CopyInstalledFilter);$(CopyLinuxFilter)"/>
 				<Property Name="CopyInstalledExceptions" Value="$(CopyInstalledExceptions);$(CopyLinuxExceptions)"/>
 			</Do>
-			<ForEach Name="Platform" Values="$(ExtensionPlatforms)">
-				<Expand Name="$(Platform)_Tag_UnrealGame"/>
-			</ForEach>
+			<Expand Name="Platform_Tag_UnrealGame"/>
 			<Do If="'$(WithHoloLens)' == true">
 				<Tag Files="#UnrealGame HoloLens;#UnrealGame HoloLens Includes" With="#Installed Win64"/>
 				<Property Name="Platforms" Value="$(Platforms)HoloLens;"/>
