Ensure the Linux toolchain is setup so it can be bundled with InstalledBuilds

diff --git a/Engine/Build/InstalledEngineBuild.xml b/Engine/Build/InstalledEngineBuild.xml
--- a/Engine/Build/InstalledEngineBuild.xml
+++ b/Engine/Build/InstalledEngineBuild.xml
@@ -1236,7 +1236,8 @@
 			<Property Name="SavedDir" Value="$(SavedOutput)\Installed\Linux"/>
 			<Copy Files="#CompressedDDCLinux" From="$(SavedDir)" To="$(LocalInstalledDirLinux)" If="'$(WithDDC)' == true And $(WithLinux)" Tag="#Installed Build Linux Files"/>
 
-			<!-- Copy the toolchain if it exists -->
+			<!-- Setup and Copy the toolchain -->
+			<Spawn Exe="$(RootDir)/Engine/Build/BatchFiles/Linux/SetupToolchain.sh"/>
 			<Property Name="LinuxToolchainDir" Value="$(ToolchainDir)/HostLinux"/>
 			<Copy Files="$(LinuxToolchainDir)/..." From="$(LinuxToolchainDir)" To="$(LocalInstalledDirLinux)/Engine/Extras/ThirdPartyNotUE/SDKs/HostLinux" Tag="#Installed Build Linux Files"/>
 
