Removing iPhonePackager.pdb from InstalledEngineBuild so it works with the Github drop.

diff --git a/Engine/Build/InstalledEngineBuild.xml b/Engine/Build/InstalledEngineBuild.xml
--- a/Engine/Build/InstalledEngineBuild.xml
+++ b/Engine/Build/InstalledEngineBuild.xml
@@ -592,11 +592,9 @@
 			<!-- Instead we will ship the binary out of Perforce for now. -->
 			<Tag Files="Engine/Binaries/DotNET/IOS/IPhonePackager.exe" With="#CopyIPPBinaries"/>
 			<Tag Files="Engine/Binaries/DotNET/IOS/IPhonePackager.exe.config" With="#CopyIPPBinaries"/>
-			<Tag Files="Engine/Binaries/DotNET/IOS/IPhonePackager.pdb" With="#CopyIPPBinaries"/>
 			<Copy Files="#CopyIPPBinaries" From="$(RootDir)" To="$(CsToolsDir)"/>
 			<Tag Files="$(CsToolsDir)/Engine/Binaries/DotNET/IOS/IPhonePackager.exe" With="#Build Tools CS Binaries Mac;#Sign Binaries Mac"/>
 			<Tag Files="$(CsToolsDir)/Engine/Binaries/DotNET/IOS/IPhonePackager.exe.config" With="#Build Tools CS Binaries Mac;#Sign Binaries Mac"/>
-			<Tag Files="$(CsToolsDir)/Engine/Binaries/DotNET/IOS/IPhonePackager.pdb" With="#Build Tools CS Binaries Mac;#Sign Binaries Mac"/>
 			
 			<!-- Sign the binaries -->
 			<Sign Files="#Sign Binaries Mac" If="'$(SignExecutables)' == true"/>
