Fix Linux not sanatizing files when making installed builds

diff --git a/Engine/Build/InstalledEngineBuild.xml b/Engine/Build/InstalledEngineBuild.xml
--- a/Engine/Build/InstalledEngineBuild.xml
+++ b/Engine/Build/InstalledEngineBuild.xml
@@ -1251,7 +1251,7 @@
 			<Tag Files="$(LocalInstalledDirLinux)/Engine/Build/InstalledBuild.txt" With="#Installed Build Linux Files"/>
 
 			<!-- Sanitize all the receipts in the target directory -->
-			<Tag Files="#Installed Build Linux Files" Filter="*.target" With="#Installed Build Linux Files"/>
+			<Tag Files="#Installed Build Linux Files" Filter="*.target" With="#Installed Build Linux Receipts"/>
 			<SanitizeReceipt Files="#Installed Build Linux Receipts" EngineDir="$(LocalInstalledDirLinux)/Engine"/>
 		</Node>
 	</Agent>
