Disable shared ddc cache for commandlet runs.
This eliminates classes of errors and init/shutdown issues that the commandlet doesn't really care about.

#rnx

diff --git a/Engine/Build/EditorPerf.xml b/Engine/Build/EditorPerf.xml
--- a/Engine/Build/EditorPerf.xml
+++ b/Engine/Build/EditorPerf.xml
@@ -81,7 +81,7 @@
 	<Property Name="ProfileArgs" Value="$(ProfileArgs) -NoShaderDistrib -VerboseShaderLogging"/>
 
 	<!-- Gauntlet arguments for commandlet, summarize runs -->
-	<Property Name="CommandletArgs" Value="-notraceserver"/>
+	<Property Name="CommandletArgs" Value="-notraceserver -ddc=NoShared"/>
 	<Property Name="SummarizeArgs" Value=""/>
 	<Property Name="SummarizeArgs" Value="-skipbaseline" If="$(RunAndBaseline)"/>
 
