Removing puzzle template for 5.0. We will rework this in a future release, but we don't want users to see it as an option in 5.0.

#rnx

diff --git a/Engine/Build/InstalledEngineFilters.xml b/Engine/Build/InstalledEngineFilters.xml
--- a/Engine/Build/InstalledEngineFilters.xml
+++ b/Engine/Build/InstalledEngineFilters.xml
@@ -39,8 +39,6 @@
 	<!-- List of project to build Feature packs -->
 	<Property Name="ProjectsToFeaturePack">
 		TP_HandheldARBP
-		TP_Puzzle
-		TP_PuzzleBP
 		TP_ThirdPerson
 		TP_ThirdPersonBP
 		TP_TopDown
@@ -63,7 +61,6 @@
 		Templates/TP_MFG_CollabBP/TP_MFG_CollabBP.uproject
 		Templates/TP_MFG_ProdConfigBP/TP_MFG_ProdConfigBP.uproject
 		Templates/TP_PhotoStudioBP/TP_PhotoStudioBP.uproject
-		Templates/TP_PuzzleBP/TP_PuzzleBP.uproject
 		Templates/TP_ThirdPersonBP/TP_ThirdPersonBP.uproject
 		Templates/TP_TopDownBP/TP_TopDownBP.uproject
 		Templates/TP_VehicleAdvBP/TP_VehicleAdvBP.uproject
@@ -115,8 +112,6 @@
 		Templates/TP_ThirdPersonBP/...
 		Templates/TP_TopDown/...
 		Templates/TP_TopDownBP/...
-		Templates/TP_Puzzle/...
-		Templates/TP_PuzzleBP/...
 		Templates/TP_VehicleAdv/...
 		Templates/TP_VehicleAdvBP/...
 		Templates/TP_VirtualRealityBP/...
