Tag Build.version for DDC Jobs

The exectuable is execut from a different spot and need this file in order for UBT to start via the editor

#rnx

diff --git a/Engine/Build/InstalledEngineBuild.xml b/Engine/Build/InstalledEngineBuild.xml
--- a/Engine/Build/InstalledEngineBuild.xml
+++ b/Engine/Build/InstalledEngineBuild.xml
@@ -750,6 +750,7 @@
 			<Tag Files="#Compile UnrealHeaderTool Win64;#UnrealEditor Win64;#Build Tools Win64" With="#ToCopy"/>
 			<Tag Files="#ToCopy" Filter="*.target" With="#TargetReceipts"/>
 			<TagReceipt Files="#TargetReceipts" RuntimeDependencies="true" With="#ToCopy"/>
+			<Tag Files="Engine\Build\Build.version" With="#ToCopy"/>
 			<Tag Files="Engine\Binaries\DotNET\Ionic.Zip.Reduced.dll" With="#ToCopy"/>
 			<Tag Files="Engine\Binaries\DotNET\OneSky.dll" With="#ToCopy"/>
 			<Tag Files="Templates\TemplateResources\..." With="#ToCopy"/>
@@ -787,6 +788,7 @@
 			<Tag Files="#Compile UnrealHeaderTool Mac;#UnrealEditor Mac;#Build Tools Mac" With="#ToCopy"/>
 			<Tag Files="#ToCopy" Filter="*.target" With="#TargetReceipts"/>
 			<TagReceipt Files="#TargetReceipts" RuntimeDependencies="true" With="#ToCopy"/>
+			<Tag Files="Engine/Build/Build.version" With="#ToCopy"/>
 			<Tag Files="Engine/Binaries/DotNET/Ionic.Zip.Reduced.dll" With="#ToCopy"/>
 			<Tag Files="Engine/Binaries/DotNET/OneSky.dll" With="#ToCopy"/>
 			<Tag Files="Templates/TemplateResources/..." With="#ToCopy"/>
@@ -825,6 +827,7 @@
 			<Tag Files="#ToCopy" Filter="*.target" With="#TargetReceipts"/>
 			<TagReceipt Files="#TargetReceipts" RuntimeDependencies="true" With="#ToCopy"/>
 			<Tag Files="Engine/Content/..." Except="*.psd;*.pdn;*.fbx;*.po" With="#ToCopy"/>
+			<Tag Files="Engine/Build/Build.version" With="#ToCopy"/>
 			<Tag Files="Engine/Config/..." Except="*.vdf" With="#ToCopy"/>
 			<Tag Files="Engine/Plugins/....uplugin;Engine/Plugins/.../Config/...;Engine/Plugins/.../Content/...;Engine/Plugins/.../Resources/...;Engine/Plugins/.../Shaders/...;Engine/Plugins/.../Templates/..." Except="Engine/Plugins/Runtime/TwitchLiveStreaming/...;Engine/Plugins/.../Binaries/Mac/..." With="#ToCopy"/>
 			<Tag Files="Engine/Shaders/..." With="#ToCopy"/>
