Adding TP_FirstPerson feature pack.

diff --git a/Engine/Build/InstalledEngineFilters.xml b/Engine/Build/InstalledEngineFilters.xml
--- a/Engine/Build/InstalledEngineFilters.xml
+++ b/Engine/Build/InstalledEngineFilters.xml
@@ -41,6 +41,8 @@
 		TP_HandheldARBP
 		TP_ThirdPerson
 		TP_ThirdPersonBP
+		TP_FirstPerson
+		TP_FirstPersonBP
 		TP_TopDown
 		TP_TopDownBP
 		TP_VehicleAdv
