Exclude the Content\Developers folder for Templates

They not need to be part of a packaged engine build.

#rnx

diff --git a/Engine/Build/InstalledEngineFilters.xml b/Engine/Build/InstalledEngineFilters.xml
--- a/Engine/Build/InstalledEngineFilters.xml
+++ b/Engine/Build/InstalledEngineFilters.xml
@@ -226,6 +226,7 @@
 		<!-- Exclude Template intermediates and other files not needed in installed build-->
 		Templates/*/Binaries/...
 		Templates/*/Build/Receipts/...
+		Templates/*/Content/Developers/...
 		Templates/*/Intermediate/...
 		Templates/*/DerivedDataCache/...
 		Templates/*/Saved/...
