Add missing OpenXR xml file to binary builds

#android

diff --git a/Engine/Build/InstalledEngineFilters.xml b/Engine/Build/InstalledEngineFilters.xml
--- a/Engine/Build/InstalledEngineFilters.xml
+++ b/Engine/Build/InstalledEngineFilters.xml
@@ -430,6 +430,7 @@
 		Engine/Binaries/DotNET/Android/...
 		Engine/Build/Android/...
 		Engine/Extras/Android/...
+		Engine/Plugins/Runtime/OpenXR/*.xml
 		Engine/Source/ThirdParty/Android/...
 		Engine/Source/ThirdParty/AndroidPermission/...
 		Engine/Source/ThirdParty/GoogleARCore/...
