Add an automated test to run the DumpMaterialShaderTypes commandlet.  This is useful for catching shader count regressions on projects.

# Changes
- Add a new test called ShaderTypeRegressionTest
- For each project:
    1. Build the editor
    2. Run the commandlet

# How do I run it?

    RunUAT BuildGraph -Script=Engine/Build/ShaderTypeRegressionTest.xml -Target="Shader Type Regression" -append:Projects="EngineTest" -append:Projects="CitySample" -append:Projects="Lyra" -append:Projects="MedievalGame[UnrealEditor]"

# Where can I find the output?

The results of the automated test are placed here:

    Engine\Programs\AutomationTool\Saved\...

diff --git a/Engine/Build/ShaderTypeRegressionTest.xml b/Engine/Build/ShaderTypeRegressionTest.xml
new file mode 100644
--- /dev/null
+++ b/Engine/Build/ShaderTypeRegressionTest.xml
@@ -0,0 +1,72 @@
+<?xml version='1.0' ?>
+<BuildGraph xmlns="http://www.epicgames.com/BuildGraph" xmlns:xsi="http://www.w3.org/2001/XMLSchema-instance" xsi:schemaLocation="http://www.epicgames.com/BuildGraph ../../../Build/Graph/Schema.xsd" >
+
+	<!-- Horde Options -->
+	<EnvVar Name="UE_HORDE_JOBID"/>
+	<Option Name="Horde" Restrict="true|false" DefaultValue="false" Description="Whether the build is running under Horde"/>
+	<Option Name="PreflightChange" DefaultValue="" Description="Preflight CL number if preflight, empty otherwise"/>
+	<Option Name="SetPreflightFlag" DefaultValue="true" Description="Whether to set the IsPreflight flag to true for preflights. Use with caution: this will cause things like build versions to be set to their non-preflight form."/>
+	<Property Name="IsPreflight" Value="false"/>
+	<Property Name="IsPreflight" Value="true" If="'$(PreflightChange)' != '' And $(SetPreflightFlag)"/>
+	<Property Name="PreflightSuffix" Value=""/>
+	<Property Name="PreflightSuffix" Value="-PF-$(PreflightChange)-$(UE_HORDE_JOBID)" If="$(IsPreflight)"/>
+	<Property Name="BuildName" Value="$(EscapedBranch)-CL-$(Change)$(PreflightSuffix)"/>
+	<Option Name="CompileCategory" DefaultValue="Editors" Description="Category string for Compile Label"/>
+	<Option Name="RunCategory" DefaultValue="Test" Description="Category string for Run Label"/>
+	
+	<!-- Projects setup-->
+	<Option Name="Projects" DefaultValue="false" Restrict=".+" Description="The projects to test"/>
+	<Error Message="No project names specified" If="'$(Projects)' == 'false'"/>
+	
+	<!-- This will be set to the names of all nodes we want to execute based on specified platforms/options -->
+	<Property Name="RunGraphNodes" Value=""/>
+	
+	<!-- This will be set to the names of all nodes we want to execute based on specified platforms/options -->
+	<Property Name="EditorPlatform" Value="Win64"/>
+
+	<!-- ForEach Project to run, generate the graph nodes -->
+	<ForEach Name="Project" Values="$(Projects)">
+
+		<!-- We parse the Editor target name from within [] in case it's not $(Project)Editor -->
+		<Property Name="EditorTarget" Value="$(Project)Editor"/>
+		<Regex Pattern="(.*)\[(.*)\]" Capture="Project;EditorTarget" Input="$(Project)" Optional="true"/>
+
+		<!-- Compile the Editor -->
+		<Property Name="CompileNodeName" Value="Compile $(Project) $(EditorPlatform)"/>
+		<Agent Name="Compile Editor Agent $(EditorPlatform) $(Project)" Type="$(EditorPlatform)">
+			<Node Name="$(CompileNodeName)">
+				<Compile Target="UnrealHeaderTool" Platform="$(EditorPlatform)" Configuration="Development"/>
+				<Compile Target="UnrealHeaderTool" Platform="$(EditorPlatform)" Project="$(Project)" Configuration="Development"/>
+				<Compile Target="ShaderCompileWorker" Platform="$(EditorPlatform)" Configuration="Development"/>
+				<Compile Target="CrashReportClientEditor" Platform="$(EditorPlatform)" Project="$(Project)" Configuration="Shipping"/>
+				<Compile Target="$(EditorTarget)" Platform="$(EditorPlatform)" Project="$(Project)" Configuration="Development"/>
+			</Node>
+		</Agent>
+		<Label Category="$(CompileCategory)" Name="EP Compile $(Project) $(EditorPlatform)" Requires="$(CompileNodeName)"/>
+		
+		<Property Name="RunNodeRequires" Value="$(CompileNodeName)"/>
+		
+		<!-- Not sure if I actually need this list -->
+		<Property Name="RunNodeList" Value=""/>
+		
+		<!-- Run the DumpMaterialShaderTypes commandlet -->
+		<Agent Name="DumpMaterialShaderTypes $(Project)" Type="Win64">
+		
+			<!-- Run the commandlet itself. -->
+			<Property Name="RunNodeName" Value="$(Project) Dump Material Shader Types"/>
+			<Node Name="$(RunNodeName)" Requires="$(RunNodeRequires);$(RunNodeList)">
+				<Property Name="Arguments" Value="-targetplatform=Windows"/>
+				<Commandlet Name="DumpMaterialShaderTypes" Project="$(Project)" Arguments="$(Arguments)"/>
+			</Node>
+			<Property Name="RunNodeList" Value="$(RunNodeList);$(RunNodeName)"/>
+			
+		</Agent>
+		<Label Category="$(RunCategory)" Name="EP Run $(Project) $(EditorPlatform)" Requires="$(RunNodeRequires);$(RunNodeList)"/>
+		<Property Name="RunGraphNodes" Value="$(RunGraphNodes);$(RunNodeList)"/>
+
+	</ForEach> <!-- Projects -->
+	
+	<!-- These Nodes run the Editor and gather profile data -->
+	<Aggregate Name="Shader Type Regression" Requires="$(RunGraphNodes)"/>
+
+</BuildGraph>
