Removing VehicleAdv template from InstalledEngineFilters. Removing CPP version only for 5.0.

diff --git a/Engine/Build/InstalledEngineFilters.xml b/Engine/Build/InstalledEngineFilters.xml
--- a/Engine/Build/InstalledEngineFilters.xml
+++ b/Engine/Build/InstalledEngineFilters.xml
@@ -45,7 +45,6 @@
 		TP_FirstPersonBP
 		TP_TopDown
 		TP_TopDownBP
-		TP_VehicleAdv
 		TP_VehicleAdvBP
 		TP_VirtualRealityBP
 	</Property>
@@ -114,7 +113,6 @@
 		Templates/TP_ThirdPersonBP/...
 		Templates/TP_TopDown/...
 		Templates/TP_TopDownBP/...
-		Templates/TP_VehicleAdv/...
 		Templates/TP_VehicleAdvBP/...
 		Templates/TP_VirtualRealityBP/...
 
