Add missing files to binary build for Android and iOS

#android
#ios

diff --git a/Engine/Build/InstalledEngineFilters.xml b/Engine/Build/InstalledEngineFilters.xml
--- a/Engine/Build/InstalledEngineFilters.xml
+++ b/Engine/Build/InstalledEngineFilters.xml
@@ -401,6 +401,9 @@
 		Engine/Source/ThirdParty/Oculus/LibOVRPlatform/*.tps
 		Engine/Source/ThirdParty/Oculus/OVRPlugin/OVRPlugin/lib/Win64/*.lib
 		Engine/Source/ThirdParty/Oculus/OVRPlugin/*.tps
+		
+		<!-- EOSSDK -->
+		Engine/Source/ThirdParty/EOSSDK/SDK/Bin/x64/xaudio2_9redist.dll
 	</Property>
 	<Property Name="CopyWin64Exceptions">
 		Engine/Source/Programs/UnrealBuildTool/obj/...
@@ -431,6 +434,8 @@
 		Engine/Plugins/Runtime/OpenXR/*.xml
 		Engine/Source/ThirdParty/Android/...
 		Engine/Source/ThirdParty/AndroidPermission/...
+		Engine/Source/ThirdParty/EOSSDK/*.xml
+		Engine/Source/ThirdParty/EOSSDK/SDK/Bin/Android/aar/*.aar
 		Engine/Source/ThirdParty/GoogleARCore/...
 		Engine/Source/ThirdParty/GoogleGameSDK/...
 		Engine/Source/ThirdParty/heapprofd/...
@@ -467,6 +472,8 @@
 		Engine/Source/ThirdParty/libJPG/...
 		Engine/Source/ThirdParty/mtlpp/...
 		Engine/Source/ThirdParty/libSampleRate/...
+		<!-- EOSSDK needs to be specified manually as we don't pick up the full .framework folder automatically -->
+		Engine/Source/ThirdParty/EOSSDK/SDK/Bin/IOS/EOSSDK.framework/...
 		<!-- PLCrashReporter needs to be specificed manually as there are two versions of the library -->
 		Engine/Source/ThirdParty/PLCrashReporter/lib/lib-Xcode-11.3.1/iOS/...
 		Engine/Source/ThirdParty/PLCrashReporter/lib/lib-Xcode-12.4/iOS/...
