BuildGraph: Add a parameter to the CsCompile task to allow using the system C# compiler rather than the UE-bundled version. Allows compiling Swarm et al correctly as part of installed build script.

diff --git a/Engine/Build/InstalledEngineBuild.xml b/Engine/Build/InstalledEngineBuild.xml
--- a/Engine/Build/InstalledEngineBuild.xml
+++ b/Engine/Build/InstalledEngineBuild.xml
@@ -143,7 +143,6 @@
 	<!-- Base editor and tools for compile on Windows -->
 	<Agent Name="Editor Win64" Type="Win64_Licensee">
 		<Node Name="Update Version Files">
-			<Log Message="$(RootDir)"/>
 			<SetVersion Change="$(Change)" Branch="$(EscapedBranch)" If="$(IsBuildMachine)"/>
 		</Node>
 		<Node Name="Compile UnrealHeaderTool Win64" Requires="Update Version Files">
@@ -518,11 +517,11 @@
 
 			<!-- Compile all the tools -->
 			<Do If="'$(HostPlatform)' == 'Win64'">
-				<CsCompile Project="$(CsToolsDir)/Engine/Source/Programs/UnrealSwarm/SwarmCoordinator/SwarmCoordinator.csproj" Configuration="Development" Platform="AnyCPU" Tag="#Build Tools CS Binaries;#Sign Binaries" TagReferences="#Build Tools CS Binaries"/>
-				<CsCompile Project="$(CsToolsDir)/Engine/Source/Programs/UnrealSwarm/Agent/Agent.csproj" Configuration="Development" Platform="AnyCPU" Tag="#Build Tools CS Binaries;#Sign Binaries" TagReferences="#Build Tools CS Binaries"/>
-				<CsCompile Project="$(CsToolsDir)/Engine/Source/Editor/SwarmInterface/DotNET/SwarmInterface.csproj" Configuration="Development" Platform="AnyCPU" Tag="#Build Tools CS Binaries;#Sign Binaries" TagReferences="#Build Tools CS Binaries"/>
-				<CsCompile Project="$(CsToolsDir)/Engine/Source/Programs/UnrealControls/UnrealControls.csproj" Configuration="Development" Platform="AnyCPU" Tag="#Build Tools CS Binaries;#Sign Binaries" TagReferences="#Build Tools CS Binaries"/>
-				<CsCompile Project="$(CsToolsDir)/Engine/Source/Programs/IOS/iPhonePackager/iPhonePackager.csproj" Configuration="Development" Platform="AnyCPU" Arguments="/verbosity:minimal /target:Rebuild" Tag="#Build Tools CS Binaries;#Sign Binaries" TagReferences="#Build Tools CS Binaries"/>
+				<CsCompile Project="$(CsToolsDir)/Engine/Source/Programs/UnrealSwarm/SwarmCoordinator/SwarmCoordinator.csproj" Configuration="Development" Platform="AnyCPU" UseSystemCompiler="true" Tag="#Build Tools CS Binaries;#Sign Binaries" TagReferences="#Build Tools CS Binaries"/>
+				<CsCompile Project="$(CsToolsDir)/Engine/Source/Programs/UnrealSwarm/Agent/Agent.csproj" Configuration="Development" Platform="AnyCPU" UseSystemCompiler="true" Tag="#Build Tools CS Binaries;#Sign Binaries" TagReferences="#Build Tools CS Binaries"/>
+				<CsCompile Project="$(CsToolsDir)/Engine/Source/Editor/SwarmInterface/DotNET/SwarmInterface.csproj" Configuration="Development" Platform="AnyCPU" UseSystemCompiler="true" Tag="#Build Tools CS Binaries;#Sign Binaries" TagReferences="#Build Tools CS Binaries"/>
+				<CsCompile Project="$(CsToolsDir)/Engine/Source/Programs/UnrealControls/UnrealControls.csproj" Configuration="Development" Platform="AnyCPU" UseSystemCompiler="true" Tag="#Build Tools CS Binaries;#Sign Binaries" TagReferences="#Build Tools CS Binaries"/>
+				<CsCompile Project="$(CsToolsDir)/Engine/Source/Programs/IOS/iPhonePackager/iPhonePackager.csproj" Configuration="Development" Platform="AnyCPU" UseSystemCompiler="true" Arguments="/verbosity:minimal /target:Rebuild" Tag="#Build Tools CS Binaries;#Sign Binaries" TagReferences="#Build Tools CS Binaries"/>
 				<Expand Name="Platform_Build_Tools_CS_Binaries"/>
 			</Do>
 			
@@ -533,8 +532,8 @@
 			<!-- Tag AutomationTool Script module build records, so that prebuilt modules may be discovered in the absence of source code -->
 			<Tag Files="$(CsToolsDir)/Engine/Intermediate/ScriptModules/..." With="#Build Tools CS Binaries" />
 
-			<CsCompile Project="$(CsToolsDir)/Engine/Source/Programs/NetworkProfiler/NetworkProfiler/NetworkProfiler.csproj" Configuration="Development" Platform="AnyCPU" Tag="#Build Tools CS Binaries" TagReferences="#Build Tools CS Binaries;#Sign Binaries" If="'$(HostPlatform)' != 'Linux'"/>
-			<CsCompile Project="$(CsToolsDir)/Engine/Source/Programs/IOS/MobileDeviceInterface/MobileDeviceInterface.csproj" Configuration="Development" Platform="AnyCPU" Arguments="/verbosity:minimal /target:Rebuild" Tag="#Build Tools CS Binaries;#Sign Binaries" TagReferences="#Build Tools CS Binaries"/>
+			<CsCompile Project="$(CsToolsDir)/Engine/Source/Programs/NetworkProfiler/NetworkProfiler/NetworkProfiler.csproj" Configuration="Development" Platform="AnyCPU" UseSystemCompiler="true" Tag="#Build Tools CS Binaries" TagReferences="#Build Tools CS Binaries;#Sign Binaries" If="'$(HostPlatform)' != 'Linux'"/>
+			<CsCompile Project="$(CsToolsDir)/Engine/Source/Programs/IOS/MobileDeviceInterface/MobileDeviceInterface.csproj" Configuration="Development" Platform="AnyCPU" UseSystemCompiler="true" Arguments="/verbosity:minimal /target:Rebuild" Tag="#Build Tools CS Binaries;#Sign Binaries" TagReferences="#Build Tools CS Binaries"/>
 
 			<!-- Create tag set for saved source needed for documentation -->
 			<Tag Files="$(CsToolsDir)/Engine/..." With="#SourceForDocumentation" />
diff --git a/Engine/Source/Programs/AutomationTool/BuildGraph/Tasks/CsCompileTask.cs b/Engine/Source/Programs/AutomationTool/BuildGraph/Tasks/CsCompileTask.cs
--- a/Engine/Source/Programs/AutomationTool/BuildGraph/Tasks/CsCompileTask.cs
+++ b/Engine/Source/Programs/AutomationTool/BuildGraph/Tasks/CsCompileTask.cs
@@ -72,6 +72,12 @@ namespace AutomationTool.Tasks
 		/// </summary>
 		[TaskParameter(Optional = true, ValidationType = TaskParameterValidationType.TagList)]
 		public string TagReferences;
+
+		/// <summary>
+		/// Whether to use the system toolchain rather than the bundled UE SDK
+		/// </summary>
+		[TaskParameter(Optional = true)]
+		public bool UseSystemCompiler;
 	}
 
 	/// <summary>
@@ -174,7 +180,14 @@ namespace AutomationTool.Tasks
 						throw new AutomationException("Project {0} does not exist!", ProjectFile);
 					}
 
-					CommandUtils.RunAndLog(CommandUtils.CmdEnv, CommandUtils.CmdEnv.DotnetMsbuildPath, $"msbuild {CommandUtils.MakePathSafeToUseWithCommandLine(ProjectFile.FullName)} {JoinedArguments}");
+					if (Parameters.UseSystemCompiler)
+					{
+						CommandUtils.MsBuild(CommandUtils.CmdEnv, ProjectFile.FullName, JoinedArguments, null);
+					}
+					else
+					{
+						CommandUtils.RunAndLog(CommandUtils.CmdEnv, CommandUtils.CmdEnv.DotnetMsbuildPath, $"msbuild {CommandUtils.MakePathSafeToUseWithCommandLine(ProjectFile.FullName)} {JoinedArguments}");
+					}
 				}
 			}
 
