Ensure Build\BatchFiles is copied for DDC Cook jobs for release builds

The Engine now invokes the Build.bat/sh for each platform

#rnx

diff --git a/Engine/Build/InstalledEngineBuild.xml b/Engine/Build/InstalledEngineBuild.xml
--- a/Engine/Build/InstalledEngineBuild.xml
+++ b/Engine/Build/InstalledEngineBuild.xml
@@ -750,6 +750,7 @@
 			<Tag Files="#ToCopy" Filter="*.target" With="#TargetReceipts"/>
 			<TagReceipt Files="#TargetReceipts" RuntimeDependencies="true" With="#ToCopy"/>
 			<Tag Files="Engine\Build\Build.version" With="#ToCopy"/>
+			<Tag Files="Engine\Build\BatchFiles\..."  With="#ToCopy"/>
 			<Tag Files="Engine\Binaries\DotNET\Ionic.Zip.Reduced.dll" With="#ToCopy"/>
 			<Tag Files="Engine\Binaries\DotNET\OneSky.dll" With="#ToCopy"/>
 			<Tag Files="Templates\TemplateResources\..." With="#ToCopy"/>
@@ -788,6 +789,7 @@
 			<Tag Files="#ToCopy" Filter="*.target" With="#TargetReceipts"/>
 			<TagReceipt Files="#TargetReceipts" RuntimeDependencies="true" With="#ToCopy"/>
 			<Tag Files="Engine/Build/Build.version" With="#ToCopy"/>
+			<Tag Files="Engine/Build/BatchFiles/Mac/..."  With="#ToCopy"/>
 			<Tag Files="Engine/Binaries/DotNET/Ionic.Zip.Reduced.dll" With="#ToCopy"/>
 			<Tag Files="Engine/Binaries/DotNET/OneSky.dll" With="#ToCopy"/>
 			<Tag Files="Templates/TemplateResources/..." With="#ToCopy"/>
@@ -827,6 +829,7 @@
 			<TagReceipt Files="#TargetReceipts" RuntimeDependencies="true" With="#ToCopy"/>
 			<Tag Files="Engine/Content/..." Except="*.psd;*.pdn;*.fbx;*.po" With="#ToCopy"/>
 			<Tag Files="Engine/Build/Build.version" With="#ToCopy"/>
+			<Tag Files="Engine/Build/BatchFiles/Linux/..."  With="#ToCopy"/>
 			<Tag Files="Engine/Config/..." Except="*.vdf" With="#ToCopy"/>
 			<Tag Files="Engine/Plugins/....uplugin;Engine/Plugins/.../Config/...;Engine/Plugins/.../Content/...;Engine/Plugins/.../Resources/...;Engine/Plugins/.../Shaders/...;Engine/Plugins/.../Templates/..." Except="Engine/Plugins/Runtime/TwitchLiveStreaming/...;Engine/Plugins/.../Binaries/Mac/..." With="#ToCopy"/>
 			<Tag Files="Engine/Shaders/..." With="#ToCopy"/>
