Templates: Re-added TP_nDisplayBP to installed templates xml.

diff --git a/Engine/Build/InstalledEngineFilters.xml b/Engine/Build/InstalledEngineFilters.xml
--- a/Engine/Build/InstalledEngineFilters.xml
+++ b/Engine/Build/InstalledEngineFilters.xml
@@ -272,9 +272,12 @@
 		Engine/Extras/UnrealVS/...
 		Engine/Extras/Redist/en-us/*
 
+		<!-- Virtual Production Templates -->
+		Templates/TP_nDisplayBP/...
+		Templates/TP_InCamVFXBP/...
+
 		<!-- Enterprise Templates Win64 Only -->
 		Templates/TP_HololensBP/...
-		Templates/TP_InCamVFXBP/...
 		Templates/TP_DMXBP/...
 
 		<!-- DotNet -->
