InstalledEngineBuild.xml: Build a few more C# projects with system compiler

#rnx

diff --git a/Engine/Build/InstalledEngineBuild.xml b/Engine/Build/InstalledEngineBuild.xml
--- a/Engine/Build/InstalledEngineBuild.xml
+++ b/Engine/Build/InstalledEngineBuild.xml
@@ -581,8 +581,8 @@
 			<!-- Tag AutomationTool Script module build records, so that prebuilt modules may be discovered in the absence of source code -->
 			<Tag Files="$(CsToolsDir)/Engine/Intermediate/ScriptModules/..." With="#Build Tools CS Binaries Mac" />
 						
-			<CsCompile Project="$(CsToolsDir)/Engine/Source/Programs/NetworkProfiler/NetworkProfiler/NetworkProfiler.csproj" Configuration="Development" Platform="AnyCPU" Tag="#Build Tools CS Binaries Mac" TagReferences="#Build Tools CS Binaries Mac;#Sign Binaries Mac"/>
-			<CsCompile Project="$(CsToolsDir)/Engine/Source/Programs/IOS/MobileDeviceInterface/MobileDeviceInterface.csproj" Configuration="Development" Platform="AnyCPU" Arguments="/verbosity:minimal /target:Rebuild" Tag="#Build Tools CS Binaries Mac;#Sign Binaries Mac" TagReferences="#Build Tools CS Binaries Mac"/>
+			<CsCompile Project="$(CsToolsDir)/Engine/Source/Programs/NetworkProfiler/NetworkProfiler/NetworkProfiler.csproj" Configuration="Development" Platform="AnyCPU" UseSystemCompiler="true" Tag="#Build Tools CS Binaries Mac" TagReferences="#Build Tools CS Binaries Mac;#Sign Binaries Mac"/>
+			<CsCompile Project="$(CsToolsDir)/Engine/Source/Programs/IOS/MobileDeviceInterface/MobileDeviceInterface.csproj" Configuration="Development" Platform="AnyCPU" UseSystemCompiler="true" Arguments="/verbosity:minimal /target:Rebuild" Tag="#Build Tools CS Binaries Mac;#Sign Binaries Mac" TagReferences="#Build Tools CS Binaries Mac"/>
 
 			<!-- IPP doesn't currently compile with the bundled Mono. -->
 			<!-- CsCompile Project="$(CsToolsDir)/Engine/Source/Programs/UnrealControls/UnrealControls.csproj" Configuration="Development" Platform="AnyCPU" Tag="#Build Tools CS Binaries Mac;#Sign Binaries Mac" TagReferences="#Build Tools CS Binaries Mac"/-->
@@ -654,7 +654,7 @@
 			<!-- Tag AutomationTool Script module build records, so that prebuilt modules may be discovered in the absence of source code -->
 			<Tag Files="$(CsToolsDir)/Engine/Intermediate/ScriptModules/..." With="#Build Tools CS Binaries LinuxNative" />
 
-			<CsCompile Project="$(CsToolsDir)/Engine/Source/Programs/IOS/MobileDeviceInterface/MobileDeviceInterface.csproj" Configuration="Development" Platform="AnyCPU" Arguments="/verbosity:minimal /target:Rebuild" Tag="#Build Tools CS Binaries LinuxNative;#Sign Binaries LinuxNative" TagReferences="#Build Tools CS Binaries LinuxNative"/>
+			<CsCompile Project="$(CsToolsDir)/Engine/Source/Programs/IOS/MobileDeviceInterface/MobileDeviceInterface.csproj" Configuration="Development" Platform="AnyCPU" UseSystemCompiler="true" Arguments="/verbosity:minimal /target:Rebuild" Tag="#Build Tools CS Binaries LinuxNative;#Sign Binaries LinuxNative" TagReferences="#Build Tools CS Binaries LinuxNative"/>
 
 			<!-- Sign the binaries -->
 			<Sign Files="#Sign Binaries LinuxNative" If="'$(SignExecutables)' == true"/>
