Optimize Shader Regression Testing Graph

- Don't use an agent for compiling and another agent for testing
- Reduces agent usage and reduces temp storage usage. each build is ~80GB

#rnx

diff --git a/Engine/Build/ShaderTypeRegressionTest.xml b/Engine/Build/ShaderTypeRegressionTest.xml
--- a/Engine/Build/ShaderTypeRegressionTest.xml
+++ b/Engine/Build/ShaderTypeRegressionTest.xml
@@ -49,27 +49,24 @@
 				<Compile Target="CrashReportClientEditor" Platform="$(EditorPlatform)" Project="$(Project)" Configuration="Shipping"/>
 				<Compile Target="$(EditorTarget)" Platform="$(EditorPlatform)" Project="$(Project)" Configuration="Development"/>
 			</Node>
-		</Agent>
-		<Label Category="$(CompileCategory)" Name="EP Compile $(Project) $(EditorPlatform)" Requires="$(CompileNodeName)"/>
-		
-		<Property Name="RunNodeRequires" Value="$(CompileNodeName)"/>
-		
-		<!-- Not sure if I actually need this list -->
-		<Property Name="RunNodeList" Value=""/>
-		
-		<!-- Run the DumpMaterialShaderTypes commandlet -->
-		<Agent Name="DumpMaterialShaderTypes $(Project)" Type="Win64">
-		
+
+			<Label Category="$(CompileCategory)" Name="EP Compile $(Project) $(EditorPlatform)" Requires="$(CompileNodeName)"/>
+			
+			<Property Name="RunNodeRequires" Value="$(CompileNodeName)"/>
+			
+			<!-- Not sure if I actually need this list -->
+			<Property Name="RunNodeList" Value=""/>
+			
 			<!-- Run the commandlet itself. -->
 			<Property Name="RunNodeName" Value="$(Project) Dump Material Shader Types"/>
 			<Node Name="$(RunNodeName)" Requires="$(RunNodeRequires);$(RunNodeList)">
 				<Commandlet Name="DumpMaterialShaderTypes" Project="$(Project)" Arguments="$(CommandletArguments)"/>
 			</Node>
 			<Property Name="RunNodeList" Value="$(RunNodeList);$(RunNodeName)"/>
-			
+
+			<Label Category="$(RunCategory)" Name="EP Run $(Project) $(EditorPlatform)" Requires="$(RunNodeRequires);$(RunNodeList)"/>
+			<Property Name="RunGraphNodes" Value="$(RunGraphNodes);$(RunNodeList)"/>
 		</Agent>
-		<Label Category="$(RunCategory)" Name="EP Run $(Project) $(EditorPlatform)" Requires="$(RunNodeRequires);$(RunNodeList)"/>
-		<Property Name="RunGraphNodes" Value="$(RunGraphNodes);$(RunNodeList)"/>
 
 	</ForEach> <!-- Projects -->
 	
