InstalledEngineBuild.xml: Don't build IOS/MobileDeviceInterface on Linux

#rnx

diff --git a/Engine/Build/InstalledEngineBuild.xml b/Engine/Build/InstalledEngineBuild.xml
--- a/Engine/Build/InstalledEngineBuild.xml
+++ b/Engine/Build/InstalledEngineBuild.xml
@@ -654,8 +654,6 @@
 			<!-- Tag AutomationTool Script module build records, so that prebuilt modules may be discovered in the absence of source code -->
 			<Tag Files="$(CsToolsDir)/Engine/Intermediate/ScriptModules/..." With="#Build Tools CS Binaries LinuxNative" />
 
-			<CsCompile Project="$(CsToolsDir)/Engine/Source/Programs/IOS/MobileDeviceInterface/MobileDeviceInterface.csproj" Configuration="Development" Platform="AnyCPU" UseSystemCompiler="true" Arguments="/verbosity:minimal /target:Rebuild" Tag="#Build Tools CS Binaries LinuxNative;#Sign Binaries LinuxNative" TagReferences="#Build Tools CS Binaries LinuxNative"/>
-
 			<!-- Sign the binaries -->
 			<Sign Files="#Sign Binaries LinuxNative" If="'$(SignExecutables)' == true"/>
 		</Node>
