Filter out more stuff for Linux installed builds

diff --git a/Engine/Build/InstalledEngineFilters.xml b/Engine/Build/InstalledEngineFilters.xml
--- a/Engine/Build/InstalledEngineFilters.xml
+++ b/Engine/Build/InstalledEngineFilters.xml
@@ -384,6 +384,14 @@
 	<Property Name="CopyEditorExceptionsLinux" Value="$(CopyEditorExceptions)">
 		FeaturePacks/TP_VirtualRealityBP.upack
 		Templates/TP_VirtualRealityBP/...
+
+		<!-- Exclude Windows/Mac binaries on Linux here -->
+		Engine/Plugins/.../Binaries/Win64/...
+		Engine/Plugins/.../Binaries/Mac/...
+
+		<!-- Don't want these folders, even if they're part of Windows/Mac tools -->
+		Engine/Binaries/Win64/...
+		Engine/Binaries/Mac/...
 	</Property>
 
 	<!-- Target Platform Filters/Exceptions -->
