Update Label strings, add a badge for Run

diff --git a/Engine/Build/EditorPerf.xml b/Engine/Build/EditorPerf.xml
--- a/Engine/Build/EditorPerf.xml
+++ b/Engine/Build/EditorPerf.xml
@@ -164,7 +164,7 @@
 					<Compile Target="$(EditorTarget)" Platform="$(Platform)" Project="$(Project)" Configuration="Development" Arguments="$(EditorCompileArgs)"/>
 				</Node>
 			</Agent>
-			<Label Category="$(CompileCategory)" Name="EP Compile $(Project) $(Platform)" Requires="$(CompileNodeName)"/>
+			<Label Category="$(CompileCategory)" Name="$(Project) Editor Compile" Requires="$(CompileNodeName)"/>
 
 			<!-- Run the Tests -->
 			<Property Name="RunNodeRequires" Value="$(CompileNodeName)"/>
@@ -183,7 +183,7 @@
 					</ForEach>
 				</ForEach>
 			</Agent>
-			<Label Category="$(RunCategory)" Name="EP Run $(Project) $(Platform)" Requires="$(RunNodeRequires);$(RunNodeList)"/>
+			<Label Category="$(RunCategory)" Name="$(Project) Editor Run" Requires="$(RunNodeRequires);$(RunNodeList)" UgsBadge="$(Project) Editor Run"/>
 			<Property Name="RunGraphNodes" Value="$(RunGraphNodes);$(RunNodeList)"/>
 
 			<!-- Baseline the Results -->
@@ -198,7 +198,7 @@
 					</ForEach>
 				</Node>
 			</Agent>
-			<Label Category="$(BaselineCategory)" Name="EP Baseline $(Project) $(Platform)" Requires="$(BaselineNodeName)"/>
+			<Label Category="$(BaselineCategory)" Name="$(Project) Perf Baseline" Requires="$(BaselineNodeName)"/>
 
 		</ForEach> <!-- Projects -->
 
