Remove filter that removes VR Template. Has always been there for Linux

diff --git a/Engine/Build/InstalledEngineFilters.xml b/Engine/Build/InstalledEngineFilters.xml
--- a/Engine/Build/InstalledEngineFilters.xml
+++ b/Engine/Build/InstalledEngineFilters.xml
@@ -382,9 +382,6 @@
 		Engine/Source/ThirdParty/Unix/LibCxx/...
 	</Property>
 	<Property Name="CopyEditorExceptionsLinux" Value="$(CopyEditorExceptions)">
-		FeaturePacks/TP_VirtualRealityBP.upack
-		Templates/TP_VirtualRealityBP/...
-
 		<!-- Exclude Windows/Mac binaries on Linux here -->
 		Engine/Plugins/.../Binaries/Win64/...
 		Engine/Plugins/.../Binaries/Mac/...
