Added SetVersion to EditorPerf job

diff --git a/Engine/Build/EditorPerf.xml b/Engine/Build/EditorPerf.xml
--- a/Engine/Build/EditorPerf.xml
+++ b/Engine/Build/EditorPerf.xml
@@ -157,6 +157,7 @@
 			<Property Name="CompileNodeName" Value="Compile $(Project) $(Platform)"/>
 			<Agent Name="Compile Editor Agent $(Platform) $(Project)" Type="$($(Platform)CompileAgentType)">
 				<Node Name="$(CompileNodeName)">
+					<SetVersion Change="$(Change)" Branch="$(EscapedBranch)" If="$(IsBuildMachine)"/>
 					<Compile Target="UnrealHeaderTool" Platform="$(Platform)" Configuration="Development"/>
 					<Compile Target="UnrealHeaderTool" Platform="$(Platform)" Project="$(Project)" Configuration="Development"/>
 					<Compile Target="ShaderCompileWorker" Platform="$(Platform)" Configuration="Development"/>
