Move HoloLens installed build script to platform extension folder & include all the binaries.

#rnx

diff --git a/Engine/Build/InstalledEngineBuild.xml b/Engine/Build/InstalledEngineBuild.xml
--- a/Engine/Build/InstalledEngineBuild.xml
+++ b/Engine/Build/InstalledEngineBuild.xml
@@ -31,9 +31,6 @@
 	<Property Name="DefaultWithIOS" Value="$(DefaultWithPlatform)"/>
 	<Property Name="DefaultWithIOS" Value="false" If="'$(HostPlatform)' != 'Mac' And !$(AllPlatforms)"/>
 
-	<Property Name="DefaultWithHoloLens" Value="$(DefaultWithPlatform)"/>
-	<Property Name="DefaultWithHoloLens" Value="false" If="'$(HostPlatform)' != 'Win64' And !$(AllPlatforms)"/>
-
 	<!-- Setup default target platforms for Installed build if they haven't been specified on the commandline -->
 	<Option Name="WithWin64" Restrict="true|false" DefaultValue="$(DefaultWithWin64)" Description="Include the Win64 target platform"/>
 	<Option Name="WithMac" Restrict="true|false" DefaultValue="$(DefaultWithMac)" Description="Include the Mac target platform"/>
@@ -42,7 +39,6 @@
 	<Option Name="WithTVOS" Restrict="true|false" DefaultValue="$(DefaultWithIOS)" Description="Include the tvOS target platform"/>
 	<Option Name="WithLinux" Restrict="true|false" DefaultValue="$(DefaultWithLinux)" Description="Include the Linux target platform"/>
 	<Option Name="WithLinuxArm64" Restrict="true|false" DefaultValue="$(DefaultWithLinuxArm64)" Description="Include the Linux AArch64 target platform"/>
-	<Option Name="WithHoloLens" Restrict="true|false" DefaultValue="$(DefaultWithHoloLens)" Description="Include the HoloLens target platform"/>
 
 	<!-- Set up target types -->
 	<Option Name="WithClient" Restrict="true|false" DefaultValue="false" Description="Include precompiled client targets"/>
@@ -89,11 +85,11 @@
 	<!-- Architectures that we build for Android. NOTE: This should match the architectures specified in BuildPluginCommand. -->
 	<Property Name="AndroidArchitectures" Value="-armv7+-arm64"/>
 	
-	<!-- Hololens Architectures -->
-	<Property Name="HoloLensArchitectures" Value="x64+arm64+multi"/>
-
 	<Property Name="AndroidGPUArchitectures" Value="es2"/>
-	
+
+	<Property Name="DefaultArchitectureCommandsWin64" Value=""/>
+
+
 	<!-- Compile flags for Windows Game binaries -->
 	<Property Name="TargetDebugInfoArg" Value="-nodebuginfo"/>
 	<Property Name="TargetDebugInfoArg" Value="" If="$(WithFullDebugInfo)"/>
@@ -372,18 +368,6 @@
 
 	<Expand Name="Platform_Agent_Target_Win64_Licensee"/>
 
-	<Agent Name="Target HoloLens" Type="HoloLens">
-		<Node Name="Compile UnrealGame HoloLens" Requires="Compile UnrealHeaderTool Win64" Produces="#UnrealGame HoloLens">
-			<ForEach Name="Configuration" Values="$(GameConfigurations)">
-				<Compile Target="UnrealGame" Platform="HoloLens" Configuration="$(Configuration)" Tag="#UnrealGame HoloLens" Arguments="-precompile -allmodules -nolink -architecture=arm64 $(TargetDebugInfoArg)"/>
-				<Compile Target="UnrealGame" Platform="HoloLens" Configuration="$(Configuration)" Tag="#UnrealGame HoloLens" Arguments="-precompile -architecture=arm64 $(TargetDebugInfoArg)" Clean="false"/>
-				<Compile Target="UnrealGame" Platform="HoloLens" Configuration="$(Configuration)" Tag="#UnrealGame HoloLens" Arguments="-precompile -allmodules -nolink -architecture=x64 $(TargetDebugInfoArg)"/>
-				<Compile Target="UnrealGame" Platform="HoloLens" Configuration="$(Configuration)" Tag="#UnrealGame HoloLens" Arguments="-precompile -architecture=x64 $(TargetDebugInfoArg)" Clean="false"/>
-			</ForEach>
-		</Node>
-	</Agent>
-	<Aggregate Name="TargetPlatforms_HoloLens" Label="Platforms/HoloLens" Requires="Compile UnrealGame HoloLens"/>
-
 	<!-- Target Platforms normally built from a Mac machine -->
 	<Agent Name="Target Mac" Type="Mac_Licensee">
 		<Property Name="OptionalTargetTags" Value=""/>
@@ -693,7 +677,6 @@
 		<Property Name="DDCPlatformsWin64" Value="$(DDCPlatformsWin64)+TVOS" If="'$(WithTVOS)' == true"/>
 		<Property Name="DDCPlatformsWin64" Value="$(DDCPlatformsWin64)+Linux" If="'$(WithLinux)' == true"/>
 		<Property Name="DDCPlatformsWin64" Value="$(DDCPlatformsWin64)+LinuxArm64" If="'$(WithLinuxArm64)' == true"/>
-		<Property Name="DDCPlatformsWin64" Value="$(DDCPlatformsWin64)+HoloLens" If="'$(WithHoloLens)' == true"/>
 		<Expand Name="Platform_Append_DDCPlatformsWin64"/>
 	</Do>
 
@@ -886,7 +869,6 @@
 		<Property Name="InstalledRequirements" Value="$(InstalledRequirements);#UnrealGame TVOS;#UnrealGame TVOS Unstripped;#UnrealGame TVOS Stripped" If="'$(WithTVOS)' == true"/>
 		<Property Name="InstalledRequirements" Value="$(InstalledRequirements);#UnrealGame Linux;#UnrealGame Linux Unstripped;#UnrealGame Linux Stripped;#Build Tools Linux" If="'$(WithLinux)' == true"/>
 		<Property Name="InstalledRequirements" Value="$(InstalledRequirements);#UnrealGame LinuxArm64;#UnrealGame LinuxArm64 Unstripped;#UnrealGame LinuxArm64 Stripped" If="'$(WithLinuxArm64)' == true"/>
-		<Property Name="InstalledRequirements" Value="$(InstalledRequirements);#UnrealGame HoloLens" If="'$(WithHoloLens)' == true"/>
 		<Property Name="InstalledRequirements" Value="$(InstalledRequirements);#CompressedDDCWin64" If="$(WithDDC) And $(WithWin64)"/>
 		<Property Name="InstalledRequirements" Value="$(InstalledRequirements);#Compile Datasmith Plugins Win64;#DatasmithPluginBinaries Win64" If="$(WithWin64) And $(CompileDatasmithPlugins)"/>
 		<!-- Feature packs -->
@@ -956,12 +938,6 @@
 				<Property Name="CopyInstalledExceptions" Value="$(CopyInstalledExceptions);$(CopyLinuxExceptions)"/>
 			</Do>
 			<Expand Name="Platform_Tag_UnrealGame"/>
-			<Do If="'$(WithHoloLens)' == true">
-				<Tag Files="#UnrealGame HoloLens;#UnrealGame HoloLens Includes" With="#Installed Win64"/>
-				<Property Name="Platforms" Value="$(Platforms)HoloLens;"/>
-				<Property Name="CopyInstalledFilter" Value="$(CopyInstalledFilter);$(CopyHoloLensFilter)"/>
-				<Property Name="CopyInstalledExceptions" Value="$(CopyInstalledExceptions);$(CopyHoloLensExceptions)"/>
-			</Do>
 			<!-- Tag any dependencies from all previous build products -->
 			<Tag Files="#Installed Win64" Filter="*.target" With="#TargetReceipts"/>
 			<TagReceipt Files="#TargetReceipts" RuntimeDependencies="true" With="#Installed Win64"/>
@@ -1010,9 +986,8 @@
 			<Property Name="ContentOnlyPlatforms" Value=""/>
 
 			<!-- Create command line arguments for any platforms that were built with additional architectures -->
-			<Property Name="ArchitectureCommands" Value=""/>
+			<Property Name="ArchitectureCommands" Value="$(DefaultArchitectureCommandsWin64)"/>
 			<Property Name="ArchitectureCommands" Value="$(ArchitectureCommands) -AndroidArchitectures=$(AndroidArchitectures)" If="'$(WithAndroid)' == true"/>
-			<Property Name="ArchitectureCommands" Value="$(ArchitectureCommands) -HoloLensArchitectures=$(HoloLensArchitectures)" If="'$(WithHoloLens)' == true"/>
 
 			<!-- Finalize the build and then tag the text file that this process creates -->
 			<Command Name="FinalizeInstalledBuild" Arguments="-OutputDir=&quot;$(LocalInstalledDir)&quot; -Platforms=$(Platforms) -ContentOnlyPlatforms=$(ContentOnlyPlatforms) $(ArchitectureCommands) -AnalyticsTypeOverride=$(AnalyticsTypeOverride)"/>
diff --git a/Engine/Build/InstalledEngineFilters.xml b/Engine/Build/InstalledEngineFilters.xml
--- a/Engine/Build/InstalledEngineFilters.xml
+++ b/Engine/Build/InstalledEngineFilters.xml
@@ -559,13 +559,4 @@
 	<Property Name="CopyLuminExceptions">
 	</Property>
 
-	<!-- HoloLens -->
-	<Property Name="CopyHoloLensFilter">
-		<!-- Build Files -->
-		Engine/Build/HoloLens/...
-	</Property>
-
-	<Property Name="CopyHoloLensExceptions">
-	</Property>
-
 </BuildGraph>
diff --git a/Engine/Platforms/Hololens/Build/InstalledEngineBuild.xml b/Engine/Platforms/Hololens/Build/InstalledEngineBuild.xml
new file mode 100644
--- /dev/null
+++ b/Engine/Platforms/Hololens/Build/InstalledEngineBuild.xml
@@ -0,0 +1,100 @@
+<?xml version='1.0' ?>
+<BuildGraph xmlns="http://www.epicgames.com/BuildGraph" xmlns:xsi="http://www.w3.org/2001/XMLSchema-instance" xsi:schemaLocation="http://www.epicgames.com/BuildGraph ../../../../Engine/Build/Graph/Schema.xsd" >
+
+	<Property Name="DefaultWithHoloLens" Value="$(DefaultWithPlatform)"/>
+	<Property Name="DefaultWithHoloLens" Value="false" If="'$(HostPlatform)' != 'Win64' And !$(AllPlatforms)"/>
+
+	<Option Name="WithHoloLens" Restrict="true|false" DefaultValue="$(DefaultWithHoloLens)" Description="Include the HoloLens target platform"/>
+
+	<!-- Hololens Architectures -->
+	<Property Name="HoloLensArchitectures" Value="x64+arm64+multi"/>
+	<Property Name="DefaultArchitectureCommandsWin64" Value="$(DefaultArchitectureCommandsWin64)  -HoloLensArchitectures=$(HoloLensArchitectures)"/>
+
+	<!-- platform extensions -->
+	<Property Name="ExtensionPlatforms" Value="$(ExtensionPlatforms);" If="'$(ExtensionPlatforms)' != ''"/>
+	<Property Name="ExtensionPlatforms" Value="$(ExtensionPlatforms)HoloLens"/>
+
+	<Do If="$(WithHoloLens)">
+		<Expand Name="AddExtensionPlatform" PlatformDir="HoloLens"/>
+	</Do>
+
+	<Extend Name="Platform_FilterRestrictedFolders">
+		<Do If="$(WithHoloLens)">
+		</Do>
+	</Extend>
+
+	<Extend Name="Platform_FilterOverrides">
+		<Do If="$(WithHoloLens)">
+			<!-- Update CopyEditorEngineOrPlatformFilter, AutomationToolPaths, etc -->
+		</Do>
+	</Extend>
+
+	<Extend Name="Platform_Agent_Target_Win64_Licensee">
+		<Agent Name="Target HoloLens" Type="HoloLens">
+			<Node Name="Compile UnrealGame HoloLens" Requires="Compile UnrealHeaderTool Win64" Produces="#UnrealGame HoloLens">
+				<ForEach Name="Configuration" Values="$(GameConfigurations)">
+					<Compile Target="UnrealGame" Platform="HoloLens" Configuration="$(Configuration)" Tag="#UnrealGame HoloLens" Arguments="-precompile -allmodules -nolink -architecture=arm64 $(TargetDebugInfoArg)"/>
+					<Compile Target="UnrealGame" Platform="HoloLens" Configuration="$(Configuration)" Tag="#UnrealGame HoloLens" Arguments="-precompile -architecture=arm64 $(TargetDebugInfoArg)" Clean="false"/>
+					<Compile Target="UnrealGame" Platform="HoloLens" Configuration="$(Configuration)" Tag="#UnrealGame HoloLens" Arguments="-precompile -allmodules -nolink -architecture=x64 $(TargetDebugInfoArg)"/>
+					<Compile Target="UnrealGame" Platform="HoloLens" Configuration="$(Configuration)" Tag="#UnrealGame HoloLens" Arguments="-precompile -architecture=x64 $(TargetDebugInfoArg)" Clean="false"/>
+				</ForEach>
+			</Node>
+		</Agent>
+		<Aggregate Name="TargetPlatforms_HoloLens" Label="Platforms/HoloLens" Requires="Compile UnrealGame HoloLens"/>
+	</Extend>
+
+	<Extend Name="Platform_AddInstalledRequirements">
+		<Do If="$(WithHoloLens)">
+			<Property Name="InstalledRequirements" Value="$(InstalledRequirements);#UnrealGame HoloLens"/>
+		</Do>
+	</Extend>
+
+	<Extend Name="Platform_TagCsToolsRedistUATSource">
+		<Do If="$(WithHoloLens)">
+		</Do>
+	</Extend>
+
+	<Extend Name="Platform_Build_Tools_CS_Binaries">
+		<Do If="$(WithHoloLens)">
+			<!-- Compile platform tools -->
+		</Do>
+	</Extend>
+
+	<Extend Name="Platform_Append_DDCPlatformsWin64">
+		<Do If="$(WithHoloLens)">
+			<Property Name="DDCPlatformsWin64" Value="$(DDCPlatformsWin64)+HoloLens"/>
+		</Do>
+	</Extend>
+
+	<Extend Name="Platform_TagDDCFiles">
+		<Do If="$(WithHoloLens)">
+			<!-- Copy utilities/tools -->
+			<Tag Files="Engine\Platforms\HoloLens\Binaries\DotNET\AutomationScripts\..." With="#ToCopy"/>
+		</Do>
+	</Extend>
+
+	<Extend Name="Platform_Tag_UnrealGame">
+		<Do If="$(WithHoloLens)">
+			<Property Name="CopyHoloLensFilter">
+				<!-- Build Files -->
+				Engine\Platforms\HoloLens\Binaries\...
+				Engine\Platforms\HoloLens\Build\...
+				Engine\Platforms\HoloLens\Content\Editor\...
+				Engine\Build\HoloLens\...
+
+				<!-- Additional Binaries -->
+				Engine\Platforms\HoloLens\Binaries\DotNET\AutomationScripts\...
+			</Property>
+
+			<Property Name="CopyHoloLensExceptions">
+			</Property>
+			
+			<Tag Files="#UnrealGame HoloLens;#UnrealGame HoloLens Includes" With="#Installed Win64"/>
+			<Property Name="Platforms" Value="$(Platforms)HoloLens;"/>
+			<Property Name="CopyInstalledFilter" Value="$(CopyInstalledFilter);$(CopyHoloLensFilter)"/>
+			<Property Name="CopyInstalledExceptions" Value="$(CopyInstalledExceptions);$(CopyHoloLensExceptions)"/>
+
+		</Do>
+	</Extend>
+
+</BuildGraph>
