Store off the generated Compressed csv files from the BuildDerivedDataCache

Helpful to diagnose and verify DDC keys part of an installed build cache

#rnx

diff --git a/Engine/Build/InstalledEngineBuild.xml b/Engine/Build/InstalledEngineBuild.xml
--- a/Engine/Build/InstalledEngineBuild.xml
+++ b/Engine/Build/InstalledEngineBuild.xml
@@ -745,6 +745,10 @@
 
 			<Command Name="BuildDerivedDataCache" Arguments="-TempDir=&quot;$(DDCDir)&quot; -FeaturePacks=&quot;$(ProjectsToBuildDDCWin64)&quot; -TargetPlatforms=$(DDCPlatformsWin64) -HostPlatform=Win64 -SavedDir=&quot;$(SavedDir)&quot;"/>
 
+			<!-- Store off the generated pack csv file for later debugging -->
+			<!-- D:\build\++UE5\Sync\LocalBuilds\InstalledDDC\Engine\DerivedDataCache\ -->
+			<Copy From="$(DDCDir)/Engine/DerivedDataCache" To="$(RootDir)/Engine/Programs/AutomationTool/Saved/Logs/DerivedDataCache" Files="Compressed-*.csv" />
+
 			<Tag Files="$(SavedDir)/Engine/DerivedDataCache/Compressed.ddp" With="#CompressedDDCWin64"/>
 		</Node>
 	</Agent>
