EditorPerf: Remove -ddc=NoShared from Commandlet runs. It was introduced to mitigate errors, and its no longer necessary.

#rnx

diff --git a/Engine/Build/EditorPerf.xml b/Engine/Build/EditorPerf.xml
--- a/Engine/Build/EditorPerf.xml
+++ b/Engine/Build/EditorPerf.xml
@@ -81,7 +81,7 @@
 	<Property Name="ProfileArgs" Value="$(ProfileArgs) -NoShaderDistrib -VerboseShaderLogging"/>
 
 	<!-- Gauntlet arguments for commandlet, summarize runs -->
-	<Property Name="CommandletArgs" Value="-notraceserver -ddc=NoShared"/>
+	<Property Name="CommandletArgs" Value="-notraceserver"/>
 	<Property Name="SummarizeArgs" Value=""/>
 	<Property Name="SummarizeArgs" Value="-skipbaseline" If="$(RunAndBaseline)"/>
 
