Modified "Installed Build Group Linux" to icnlude the result of the compilation of the Datasmith plugins on Linux with a full workspace
Added "Extra Plugins Agent Linux" agent to compile the Datasmith plugins on Linux with a full workspace

diff --git a/Engine/Build/InstalledEngineBuild.xml b/Engine/Build/InstalledEngineBuild.xml
--- a/Engine/Build/InstalledEngineBuild.xml
+++ b/Engine/Build/InstalledEngineBuild.xml
@@ -221,11 +221,12 @@
 		<Node Name="Compile UnrealHeaderTool Linux" Requires="Update Version Files">
 			<Compile Target="UnrealHeaderTool" Platform="Linux" Configuration="Development" Arguments="-precompile -allmodules"/>
 		</Node>
-		<Node Name="Compile UnrealEditor Linux" Requires="Compile UnrealHeaderTool Linux" Produces="#UnrealEditor Linux;#UnrealEditor Linux Unstripped;#UnrealEditor Linux Stripped">
+		<Node Name="Compile UnrealEditor Linux" Requires="Compile UnrealHeaderTool Linux" Produces="#UnrealEditor Linux;#UnrealEditor Linux Unstripped;#UnrealEditor Linux Stripped;#UnrealEditor Linux WithExceptions">
 			<Compile Target="UnrealEditor" Platform="Linux" Configuration="DebugGame" Tag="#UnrealEditor Linux" Arguments="-precompile -allmodules $(CrashReporterCompileArgs)"/>
 			<Compile Target="UnrealEditor" Platform="Linux" Configuration="Development" Tag="#UnrealEditor Linux" Arguments="-precompile -allmodules $(CrashReporterCompileArgs)"/>
 
 			<Tag Files="#UnrealEditor Linux" Filter="$(LinuxStripFilter)" With="#UnrealEditor Linux Unstripped"/>
+			<Tag Files="#UnrealEditor Linux" Except="$(PluginsExceptions)" With="#UnrealEditor Linux WithExceptions"/>
 
 			<Switch>
 				<Case If="'$(WithFullDebugInfo)' == true">
@@ -835,14 +836,14 @@
 		</Node>
 	</Agent>
 
-	<!-- These are plugins that need to be compiled with a full workspace rather than Win64_Licensee. -->
+	<!-- These are plugins that need to be compiled with a full workspace rather than $(Platform)_Licensee. -->
 	<Agent Name="Extra Plugins Agent Win64" Type="Win64">
 		<Node Name="Compile Datasmith Plugins Win64" Requires="Compile UnrealHeaderTool Win64" Produces="#DatasmithPluginBinaries Win64">
 			<Do If="$(WithWin64) And $(CompileDatasmithPlugins) == 'True'">
 				<Property Name="DatasmithPlugins" Value="DatasmithCADImporter;DatasmithIFCImporter;DatasmithC4DImporter;AxFImporter;MDLImporter"/>
 				<ForEach Name="Plugin" Values="$(DatasmithPlugins)">
 					<Compile Target="UnrealEditor" Configuration="Development" Platform="Win64" Arguments="-Plugin=&quot;$(RootDir)/Engine/Plugins/Enterprise/$(Plugin)/$(Plugin).uplugin&quot;"/>
-					<Tag Files="$(RootDir)/Engine/Plugins/Enterprise/$(Plugin)/Binaries/..." With="#DatasmithPluginBinaries Win64"/>
+					<Tag Files="$(RootDir)/Engine/Plugins/Enterprise/$(Plugin)/Binaries/Win64/..." With="#DatasmithPluginBinaries Win64"/>
 				</ForEach>
 			</Do>
 		</Node>
@@ -854,7 +855,19 @@
 				<Property Name="DatasmithPlugins" Value="DatasmithCADImporter;DatasmithIFCImporter;DatasmithC4DImporter;AxFImporter;MDLImporter"/>
 				<ForEach Name="Plugin" Values="$(DatasmithPlugins)">
 					<Compile Target="UnrealEditor" Configuration="Development" Platform="Mac" Arguments="-Plugin=&quot;$(RootDir)/Engine/Plugins/Enterprise/$(Plugin)/$(Plugin).uplugin&quot;"/>
-					<Tag Files="$(RootDir)/Engine/Plugins/Enterprise/$(Plugin)/Binaries/..." With="#DatasmithPluginBinaries Mac"/>
+					<Tag Files="$(RootDir)/Engine/Plugins/Enterprise/$(Plugin)/Binaries/Mac/..." With="#DatasmithPluginBinaries Mac"/>
+				</ForEach>
+			</Do>
+		</Node>
+	</Agent>
+
+	<Agent Name="Extra Plugins Agent Linux" Type="Linux">
+		<Node Name="Compile Datasmith Plugins Linux" Requires="Compile UnrealHeaderTool Linux" Produces="#DatasmithPluginBinaries Linux">
+			<Do If="$(WithLinux) And $(CompileDatasmithPlugins) == 'True'">
+				<Property Name="DatasmithPlugins" Value="DatasmithCADImporter"/>
+				<ForEach Name="Plugin" Values="$(DatasmithPlugins)">
+					<Compile Target="UnrealEditor" Configuration="Development" Platform="Linux" Arguments="-Plugin=&quot;$(RootDir)/Engine/Plugins/Enterprise/$(Plugin)/$(Plugin).uplugin&quot;"/>
+					<Tag Files="$(RootDir)/Engine/Plugins/Enterprise/$(Plugin)/Binaries/Linux/..." With="#DatasmithPluginBinaries Linux"/>
 				</ForEach>
 			</Do>
 		</Node>
@@ -1130,13 +1143,16 @@
 
 	<Agent Name="Installed Build Group Linux" Type="Linux_Licensee">
 		<!-- Build up a list of requirements as some can be optional -->
-		<Property Name="InstalledRequirements" Value="#Compile UnrealHeaderTool Linux;#UnrealEditor Linux;#UnrealEditor Linux Unstripped;#UnrealEditor Linux Stripped;#Build Tools LinuxNative;#Build Tools CS Binaries LinuxNative;#Rules Assemblies"/>
+		<Property Name="UnrealEditor" Value="#UnrealEditor Linux" />
+		<Property Name="UnrealEditor" Value="#UnrealEditor Linux WithExceptions" If="$(CompileDatasmithPlugins)" />
+		<Property Name="InstalledRequirements" Value="#Compile UnrealHeaderTool Linux;$(UnrealEditor);#UnrealEditor Linux Unstripped;#UnrealEditor Linux Stripped;#Build Tools LinuxNative;#Build Tools CS Binaries LinuxNative;#Rules Assemblies"/>
 
 		<!-- Optional Target Platform requirements -->
 		<Property Name="InstalledRequirements" Value="$(InstalledRequirements);#UnrealGame LinuxNative;#UnrealGame LinuxNative Unstripped;#UnrealGame LinuxNative Stripped" If="'$(WithLinux)' == true"/>
 		<Property Name="InstalledRequirements" Value="$(InstalledRequirements);#UnrealGame LinuxArm64 LinuxNative;#UnrealGame LinuxArm64 LinuxNative Unstripped;#UnrealGame LinuxArm64 LinuxNative Stripped" If="'$(WithLinuxArm64)' == true"/>
 		<Property Name="InstalledRequirements" Value="$(InstalledRequirements);#UnrealGame Android;#UnrealGame Android Unstripped;#UnrealGame Android Stripped" If="'$(WithAndroid)' == true"/>
 		<Property Name="InstalledRequirements" Value="$(InstalledRequirements);#CompressedDDCLinux" If="'$(WithDDC)' == true And $(WithLinux)"/>
+		<Property Name="InstalledRequirements" Value="$(InstalledRequirements);#Compile Datasmith Plugins Linux;#DatasmithPluginBinaries Linux" If="$(WithLinux) And $(CompileDatasmithPlugins)"/>
 
 		<!-- Feature packs -->
 		<Property Name="InstalledRequirements" Value="$(InstalledRequirements);Make Feature Packs LinuxNative"/>
@@ -1156,7 +1172,7 @@
 
 			<!-- Tag files that always need to be part of an installed build -->
 			<Tag Files="#Compile UnrealHeaderTool Linux" With="#Installed Linux"/>
-			<Tag Files="#UnrealEditor Linux" Except="#UnrealEditor Linux Unstripped" With="#Installed Linux"/>
+			<Tag Files="$(UnrealEditor)" Except="#UnrealEditor Linux Unstripped" With="#Installed Linux"/>
 			<Tag Files="#UnrealEditor Linux Stripped" With="#Saved Output"/>
 			<Tag Files="#Build Tools LinuxNative;#Make Feature Packs LinuxNative;#Rules Assemblies" With="#Installed Linux"/>
 
@@ -1164,6 +1180,7 @@
 			<Do If="'$(WithLinux)' == true">
 				<Tag Files="#UnrealGame LinuxNative;#Build Tools LinuxNative" Except="#UnrealGame LinuxNative Unstripped" With="#Installed Linux"/>
 				<Tag Files="#UnrealGame LinuxNative Stripped" With="#Saved Output"/>
+				<Tag Files="#DatasmithPluginBinaries Linux" With="#Installed Linux" If="$(WithLinux) And $(CompileDatasmithPlugins)"/>
 				<Property Name="Platforms" Value="$(Platforms)Linux;"/>
 				<Property Name="CopyInstalledFilter" Value="$(CopyInstalledFilter);$(CopyLinuxFilter)"/>
 				<Property Name="CopyInstalledExceptions" Value="$(CopyInstalledExceptions);$(CopyLinuxExceptions)"/>
