Fixed missing MDLImporter module in Linux binary build

diff --git a/Engine/Build/InstalledEngineBuild.xml b/Engine/Build/InstalledEngineBuild.xml
--- a/Engine/Build/InstalledEngineBuild.xml
+++ b/Engine/Build/InstalledEngineBuild.xml
@@ -864,7 +864,7 @@
 	<Agent Name="Extra Plugins Agent Linux" Type="Linux">
 		<Node Name="Compile Datasmith Plugins Linux" Requires="Compile UnrealHeaderTool Linux" Produces="#DatasmithPluginBinaries Linux">
 			<Do If="$(WithLinux) And $(CompileDatasmithPlugins) == 'True'">
-				<Property Name="DatasmithPlugins" Value="DatasmithCADImporter"/>
+				<Property Name="DatasmithPlugins" Value="DatasmithCADImporter;MDLImporter"/>
 				<ForEach Name="Plugin" Values="$(DatasmithPlugins)">
 					<Compile Target="UnrealEditor" Configuration="Development" Platform="Linux" Arguments="-Plugin=&quot;$(RootDir)/Engine/Plugins/Enterprise/$(Plugin)/$(Plugin).uplugin&quot;"/>
 					<Tag Files="$(RootDir)/Engine/Plugins/Enterprise/$(Plugin)/Binaries/Linux/..." With="#DatasmithPluginBinaries Linux"/>
