Add Arm64 support for Mac builds.

#rnx

diff --git a/Engine/Build/InstalledEngineBuild.xml b/Engine/Build/InstalledEngineBuild.xml
--- a/Engine/Build/InstalledEngineBuild.xml
+++ b/Engine/Build/InstalledEngineBuild.xml
@@ -34,6 +34,7 @@
 	<!-- Setup default target platforms for Installed build if they haven't been specified on the commandline -->
 	<Option Name="WithWin64" Restrict="true|false" DefaultValue="$(DefaultWithWin64)" Description="Include the Win64 target platform"/>
 	<Option Name="WithMac" Restrict="true|false" DefaultValue="$(DefaultWithMac)" Description="Include the Mac target platform"/>
+	<Option Name="WithMacArm64" Restrict="true|false" DefaultValue="false" Description="Include the Arm64 architecture version of the Mac target platform"/>
 	<Option Name="WithAndroid" Restrict="true|false" DefaultValue="$(DefaultWithPlatform)" Description="Include the Android target platform"/>
 	<Option Name="WithIOS" Restrict="true|false" DefaultValue="$(DefaultWithIOS)" Description="Include the iOS target platform"/>
 	<Option Name="WithTVOS" Restrict="true|false" DefaultValue="$(DefaultWithIOS)" Description="Include the tvOS target platform"/>
@@ -136,6 +137,9 @@
 	<Property Name="CrashReporterCompileArgs" Value=""/>
 	<Property Name="CrashReporterCompileArgs" Value="-define:CRC_TELEMETRY_URL=\&quot;$(CrashReporterAPIURL)\&quot; -define:CRC_TELEMETRY_KEY_DEV=\&quot;$(CrashReporterAPIKey)\&quot; -define:CRC_TELEMETRY_KEY_RELEASE=\&quot;$(CrashReporterAPIKey)\&quot; -OverrideBuildEnvironment" If="'$(CrashReporterAPIURL)' != '' And '$(CrashReporterAPIKey)' != ''" />
 
+	<Property Name="UseArm64ArchitectureArg" Value="" />
+	<Property Name="UseArm64ArchitectureArg" Value="-architecture=arm64" If="$(WithMacArm64)" />
+
 	<!-- Base editor and tools for compile on Windows -->
 	<Agent Name="Editor Win64" Type="Win64_Licensee">
 		<Node Name="Update Version Files">
@@ -187,8 +191,8 @@
 			<Compile Target="UnrealHeaderTool" Platform="Mac" Configuration="Development" Arguments="-precompile -allmodules"/>
 		</Node>
 		<Node Name="Compile UnrealEditor Mac" Requires="Compile UnrealHeaderTool Mac" Produces="#UnrealEditor Mac;#UnrealEditor Mac Unstripped;#UnrealEditor Mac Stripped;#UnrealEditor Mac Unsigned;#UnrealEditor Mac Signed;#UnrealEditor Mac WithExceptions">
-			<Compile Target="UnrealEditor" Platform="Mac" Configuration="DebugGame" Tag="#UnrealEditor Mac" Arguments="-precompile -allmodules $(CrashReporterCompileArgs)"/>
-			<Compile Target="UnrealEditor" Platform="Mac" Configuration="Development" Tag="#UnrealEditor Mac" Arguments="-precompile -allmodules $(CrashReporterCompileArgs)"/>
+			<Compile Target="UnrealEditor" Platform="Mac" Configuration="DebugGame" Tag="#UnrealEditor Mac" Arguments="-precompile -allmodules $(CrashReporterCompileArgs) $(UseArm64ArchitectureArg)"/>
+			<Compile Target="UnrealEditor" Platform="Mac" Configuration="Development" Tag="#UnrealEditor Mac" Arguments="-precompile -allmodules $(CrashReporterCompileArgs) $(UseArm64ArchitectureArg)"/>
 
 			<Property Name="UnrealEditorExceptions" Value=""/>
 			<Property Name="UnrealEditorExceptions" Value="$(PluginsExceptions)" If="$(CompileDatasmithPlugins)"/>
@@ -585,20 +589,20 @@
 		</Node>
 	
 		<Node Name="Build Tools Mac" Requires="Compile UnrealHeaderTool Mac">
-			<Compile Target="CrashReportClient" Configuration="Shipping" Platform="Mac" Tag="#Build Tools Mac" Arguments="$(CrashReporterCompileArgs)"/>
-			<Compile Target="CrashReportClientEditor" Configuration="Shipping" Platform="Mac" Tag="#Build Tools Mac" Arguments="$(CrashReporterCompileArgs)"/>
-			<Compile Target="ShaderCompileWorker" Configuration="Development" Platform="Mac" Tag="#Build Tools Mac"/>
-			<Compile Target="EpicWebHelper" Configuration="Development" Platform="Mac" Tag="#Build Tools Mac"/>
-			<Compile Target="UnrealInsights" Configuration="Development" Platform="Mac" Tag="#Build Tools Mac"/>
-			<Compile Target="UnrealFrontend" Configuration="Development" Platform="Mac" Tag="#Build Tools Mac"/>
-			<Compile Target="UnrealLightmass" Configuration="Development" Platform="Mac" Tag="#Build Tools Mac"/>
-			<Compile Target="UnrealPak" Configuration="Development" Platform="Mac" Tag="#Build Tools Mac"/>
-			<Compile Target="UnrealMultiUserServer" Configuration="Development" Platform="Mac" Tag="#Build Tools Mac"/>
-			<Compile Target="UnrealRecoverySvc" Configuration="Development" Platform="Mac" Tag="#Build Tools Mac"/>
-			<Compile Target="BuildPatchTool" Configuration="Shipping" Platform="Mac" Tag="#Build Tools Mac"/>
-			<Compile Target="DsymExporter" Configuration="Development" Platform="Mac" Tag="#Build Tools Mac"/>
-			<Compile Target="UnrealEditorServices" Configuration="Development" Platform="Mac" Tag="#Build Tools Mac"/>
-			<Compile Target="UnrealAtoS" Configuration="Development" Platform="Mac" Tag="#Build Tools Mac"/>
+			<Compile Target="CrashReportClient" Configuration="Shipping" Platform="Mac" Tag="#Build Tools Mac" Arguments="$(CrashReporterCompileArgs) $(UseArm64ArchitectureArg)"/>
+			<Compile Target="CrashReportClientEditor" Configuration="Shipping" Platform="Mac" Tag="#Build Tools Mac" Arguments="$(CrashReporterCompileArgs) $(UseArm64ArchitectureArg)"/>
+			<Compile Target="ShaderCompileWorker" Configuration="Development" Platform="Mac" Tag="#Build Tools Mac" Arguments=" $(UseArm64ArchitectureArg)" />
+			<Compile Target="EpicWebHelper" Configuration="Development" Platform="Mac" Tag="#Build Tools Mac" Arguments=" $(UseArm64ArchitectureArg)" />
+			<Compile Target="UnrealInsights" Configuration="Development" Platform="Mac" Tag="#Build Tools Mac" Arguments=" $(UseArm64ArchitectureArg)" />
+			<Compile Target="UnrealFrontend" Configuration="Development" Platform="Mac" Tag="#Build Tools Mac" Arguments=" $(UseArm64ArchitectureArg)" />
+			<Compile Target="UnrealLightmass" Configuration="Development" Platform="Mac" Tag="#Build Tools Mac" Arguments=" $(UseArm64ArchitectureArg)" />
+			<Compile Target="UnrealPak" Configuration="Development" Platform="Mac" Tag="#Build Tools Mac" Arguments=" $(UseArm64ArchitectureArg)" />
+			<Compile Target="UnrealMultiUserServer" Configuration="Development" Platform="Mac" Tag="#Build Tools Mac" Arguments=" $(UseArm64ArchitectureArg)" />
+			<Compile Target="UnrealRecoverySvc" Configuration="Development" Platform="Mac" Tag="#Build Tools Mac" Arguments=" $(UseArm64ArchitectureArg)" />
+			<Compile Target="BuildPatchTool" Configuration="Shipping" Platform="Mac" Tag="#Build Tools Mac" Arguments=" $(UseArm64ArchitectureArg)" />
+			<Compile Target="DsymExporter" Configuration="Development" Platform="Mac" Tag="#Build Tools Mac" Arguments=" $(UseArm64ArchitectureArg)" />
+			<Compile Target="UnrealEditorServices" Configuration="Development" Platform="Mac" Tag="#Build Tools Mac" Arguments=" $(UseArm64ArchitectureArg)" />
+			<Compile Target="UnrealAtoS" Configuration="Development" Platform="Mac" Tag="#Build Tools Mac" Arguments=" $(UseArm64ArchitectureArg)" />
 		</Node>
 	</Agent>
 
