Only build UnrealEditor Mac for Arm64 when WithMacArm64 is set.

#rnx

diff --git a/Engine/Build/InstalledEngineBuild.xml b/Engine/Build/InstalledEngineBuild.xml
--- a/Engine/Build/InstalledEngineBuild.xml
+++ b/Engine/Build/InstalledEngineBuild.xml
@@ -589,20 +589,20 @@
 		</Node>
 	
 		<Node Name="Build Tools Mac" Requires="Compile UnrealHeaderTool Mac">
-			<Compile Target="CrashReportClient" Configuration="Shipping" Platform="Mac" Tag="#Build Tools Mac" Arguments="$(CrashReporterCompileArgs) $(UseArm64ArchitectureArg)"/>
-			<Compile Target="CrashReportClientEditor" Configuration="Shipping" Platform="Mac" Tag="#Build Tools Mac" Arguments="$(CrashReporterCompileArgs) $(UseArm64ArchitectureArg)"/>
-			<Compile Target="ShaderCompileWorker" Configuration="Development" Platform="Mac" Tag="#Build Tools Mac" Arguments=" $(UseArm64ArchitectureArg)" />
-			<Compile Target="EpicWebHelper" Configuration="Development" Platform="Mac" Tag="#Build Tools Mac" Arguments=" $(UseArm64ArchitectureArg)" />
-			<Compile Target="UnrealInsights" Configuration="Development" Platform="Mac" Tag="#Build Tools Mac" Arguments=" $(UseArm64ArchitectureArg)" />
-			<Compile Target="UnrealFrontend" Configuration="Development" Platform="Mac" Tag="#Build Tools Mac" Arguments=" $(UseArm64ArchitectureArg)" />
-			<Compile Target="UnrealLightmass" Configuration="Development" Platform="Mac" Tag="#Build Tools Mac" Arguments=" $(UseArm64ArchitectureArg)" />
-			<Compile Target="UnrealPak" Configuration="Development" Platform="Mac" Tag="#Build Tools Mac" Arguments=" $(UseArm64ArchitectureArg)" />
-			<Compile Target="UnrealMultiUserServer" Configuration="Development" Platform="Mac" Tag="#Build Tools Mac" Arguments=" $(UseArm64ArchitectureArg)" />
-			<Compile Target="UnrealRecoverySvc" Configuration="Development" Platform="Mac" Tag="#Build Tools Mac" Arguments=" $(UseArm64ArchitectureArg)" />
-			<Compile Target="BuildPatchTool" Configuration="Shipping" Platform="Mac" Tag="#Build Tools Mac" Arguments=" $(UseArm64ArchitectureArg)" />
-			<Compile Target="DsymExporter" Configuration="Development" Platform="Mac" Tag="#Build Tools Mac" Arguments=" $(UseArm64ArchitectureArg)" />
-			<Compile Target="UnrealEditorServices" Configuration="Development" Platform="Mac" Tag="#Build Tools Mac" Arguments=" $(UseArm64ArchitectureArg)" />
-			<Compile Target="UnrealAtoS" Configuration="Development" Platform="Mac" Tag="#Build Tools Mac" Arguments=" $(UseArm64ArchitectureArg)" />
+			<Compile Target="CrashReportClient" Configuration="Shipping" Platform="Mac" Tag="#Build Tools Mac" Arguments="$(CrashReporterCompileArgs)"/>
+			<Compile Target="CrashReportClientEditor" Configuration="Shipping" Platform="Mac" Tag="#Build Tools Mac" Arguments="$(CrashReporterCompileArgs)"/>
+			<Compile Target="ShaderCompileWorker" Configuration="Development" Platform="Mac" Tag="#Build Tools Mac"/>
+			<Compile Target="EpicWebHelper" Configuration="Development" Platform="Mac" Tag="#Build Tools Mac"/>
+			<Compile Target="UnrealInsights" Configuration="Development" Platform="Mac" Tag="#Build Tools Mac"/>
+			<Compile Target="UnrealFrontend" Configuration="Development" Platform="Mac" Tag="#Build Tools Mac"/>
+			<Compile Target="UnrealLightmass" Configuration="Development" Platform="Mac" Tag="#Build Tools Mac"/>
+			<Compile Target="UnrealPak" Configuration="Development" Platform="Mac" Tag="#Build Tools Mac"/>
+			<Compile Target="UnrealMultiUserServer" Configuration="Development" Platform="Mac" Tag="#Build Tools Mac"/>
+			<Compile Target="UnrealRecoverySvc" Configuration="Development" Platform="Mac" Tag="#Build Tools Mac"/>
+			<Compile Target="BuildPatchTool" Configuration="Shipping" Platform="Mac" Tag="#Build Tools Mac"/>
+			<Compile Target="DsymExporter" Configuration="Development" Platform="Mac" Tag="#Build Tools Mac"/>
+			<Compile Target="UnrealEditorServices" Configuration="Development" Platform="Mac" Tag="#Build Tools Mac"/>
+			<Compile Target="UnrealAtoS" Configuration="Development" Platform="Mac" Tag="#Build Tools Mac"/>
 		</Node>
 	</Agent>
 
