Build DDC Against an Arm64 Mac if WithMacArm64 is true

#rnx

diff --git a/Engine/Build/InstalledEngineBuild.xml b/Engine/Build/InstalledEngineBuild.xml
--- a/Engine/Build/InstalledEngineBuild.xml
+++ b/Engine/Build/InstalledEngineBuild.xml
@@ -768,7 +768,9 @@
 		<Property Name="DDCPlatformsMac" Value="$(DDCPlatformsMac)+TVOS" If="'$(WithTVOS)' == true"/>
 	</Do>
 
-	<Agent Name="DDC Group Mac" Type="Mac">
+	<Property Name="MacDDCAgentType" Value="Mac" />
+	<Property Name="MacDDCAgentType" Value="MacArm64" If="$(WithMacArm64)" />
+	<Agent Name="DDC Group Mac" Type="$(MacDDCAgentType)">
 		<Node Name="Build DDC Mac" Requires="#UnrealEditor Mac;Build Tools Mac;#Build Tools CS Binaries" Produces="#CompressedDDCMac">
 			<!-- Get our temp dir -->
 			<Property Name="DDCDir" Value="$(RootDir)/LocalBuilds/InstalledDDC"/>
