Update bundled dotnet to 6.0.302

#rnx

diff --git a/Engine/Build/BatchFiles/GetDotnetPath.bat b/Engine/Build/BatchFiles/GetDotnetPath.bat
--- a/Engine/Build/BatchFiles/GetDotnetPath.bat
+++ b/Engine/Build/BatchFiles/GetDotnetPath.bat
@@ -9,7 +9,7 @@ rem if UE_USE_SYSTEM_DOTNET we assume a installed dotnet is present
 if "%UE_USE_SYSTEM_DOTNET%" == "1" goto verify_dotnet
 
 rem add the dotnet sdk in the repo as the current dotnet sdk
-set UE_DOTNET_DIR=%~dp0..\..\Binaries\ThirdParty\DotNet\6.0.200\windows
+set UE_DOTNET_DIR=%~dp0..\..\Binaries\ThirdParty\DotNet\6.0.302\windows
 set PATH=%UE_DOTNET_DIR%;%PATH%
 set DOTNET_ROOT=%UE_DOTNET_DIR%
 
diff --git a/Engine/Build/BatchFiles/Linux/SetupDotnet.sh b/Engine/Build/BatchFiles/Linux/SetupDotnet.sh
--- a/Engine/Build/BatchFiles/Linux/SetupDotnet.sh
+++ b/Engine/Build/BatchFiles/Linux/SetupDotnet.sh
@@ -31,7 +31,7 @@ fi
 if [ $IS_DOTNET_INSTALLED -eq 0 ]; then
 	echo Setting up bundled DotNet SDK
 	CUR_DIR=`pwd`
-	export UE_DOTNET_DIR="$CUR_DIR/../../../Binaries/ThirdParty/DotNet/6.0.200/linux"
+	export UE_DOTNET_DIR="$CUR_DIR/../../../Binaries/ThirdParty/DotNet/6.0.302/linux"
 	chmod u+x "$UE_DOTNET_DIR/dotnet"
 	export PATH="$UE_DOTNET_DIR:$PATH"
 	export DOTNET_ROOT="$UE_DOTNET_DIR"
diff --git a/Engine/Build/BatchFiles/Mac/SetupDotnet.sh b/Engine/Build/BatchFiles/Mac/SetupDotnet.sh
--- a/Engine/Build/BatchFiles/Mac/SetupDotnet.sh
+++ b/Engine/Build/BatchFiles/Mac/SetupDotnet.sh
@@ -36,7 +36,7 @@ if [ $IS_DOTNET_INSTALLED -eq 0 ]; then
 	ARCH=x64
 	[ $(uname -m) == "arm64" ] && ARCH=arm64 
 	
-	export UE_DOTNET_DIR=$CUR_DIR/../../../Binaries/ThirdParty/DotNet/6.0.200/mac-$ARCH
+	export UE_DOTNET_DIR=$CUR_DIR/../../../Binaries/ThirdParty/DotNet/6.0.302/mac-$ARCH
 	chmod u+x "$UE_DOTNET_DIR/dotnet"
 	echo $UE_DOTNET_DIR
 	export PATH=$UE_DOTNET_DIR:$PATH
diff --git a/Engine/Build/InstalledEngineFilters.xml b/Engine/Build/InstalledEngineFilters.xml
--- a/Engine/Build/InstalledEngineFilters.xml
+++ b/Engine/Build/InstalledEngineFilters.xml
@@ -281,7 +281,7 @@
 		Templates/TP_DMXBP/...
 
 		<!-- DotNet -->
-		Engine/Binaries/ThirdParty/DotNet/6.0.200/windows/...
+		Engine/Binaries/ThirdParty/DotNet/6.0.302/windows/...
 		Engine/Binaries/ThirdParty/VisualStudio/...
 
 		<!-- CsvTools binaries which are manually updated in P4 -->
@@ -326,8 +326,8 @@
 		Engine/Binaries/ThirdParty/Mono/Mac/...
 
 		<!-- DotNet -->
-		Engine/Binaries/ThirdParty/DotNet/6.0.200/mac-arm64/...
-		Engine/Binaries/ThirdParty/DotNet/6.0.200/mac-x64/...
+		Engine/Binaries/ThirdParty/DotNet/6.0.302/mac-arm64/...
+		Engine/Binaries/ThirdParty/DotNet/6.0.302/mac-x64/...
 		Engine/Binaries/ThirdParty/VisualStudio/...
 	</Property>
 
@@ -372,7 +372,7 @@
 		Engine/Binaries/ThirdParty/Mono/...
 
 		<!-- DotNet -->
-		Engine/Binaries/ThirdParty/DotNet/6.0.200/linux/...
+		Engine/Binaries/ThirdParty/DotNet/6.0.302/linux/...
 		Engine/Binaries/ThirdParty/VisualStudio/...
 
 		<!-- DotNet bundled Requirements -->
diff --git a/Engine/Build/Rsync/RsyncEngine.txt b/Engine/Build/Rsync/RsyncEngine.txt
--- a/Engine/Build/Rsync/RsyncEngine.txt
+++ b/Engine/Build/Rsync/RsyncEngine.txt
@@ -62,9 +62,9 @@ protect /Binaries/TVOS/
 # Binaries/ThirdParty, just DotNet and VisualStudio required to run UnrealBuildTool, and libtbb used by UHT
 + /Binaries/ThirdParty/
 + /Binaries/ThirdParty/DotNet/
-+ /Binaries/ThirdParty/DotNet/6.0.200/
-+ /Binaries/ThirdParty/DotNet/6.0.200/mac**/
-+ /Binaries/ThirdParty/DotNet/6.0.200/mac**/**
++ /Binaries/ThirdParty/DotNet/6.0.302/
++ /Binaries/ThirdParty/DotNet/6.0.302/mac**/
++ /Binaries/ThirdParty/DotNet/6.0.302/mac**/**
 - /Binaries/ThirdParty/DotNet/**
 + /Binaries/ThirdParty/VisualStudio/
 + /Binaries/ThirdParty/VisualStudio/**
diff --git a/Engine/Source/Programs/Shared/EpicGames.Build/Unreal.cs b/Engine/Source/Programs/Shared/EpicGames.Build/Unreal.cs
--- a/Engine/Source/Programs/Shared/EpicGames.Build/Unreal.cs
+++ b/Engine/Source/Programs/Shared/EpicGames.Build/Unreal.cs
@@ -79,7 +79,7 @@ namespace UnrealBuildBase
 			return UnrealBuildToolDllPath;
 		}
 
-		static private string DotnetVersionDirectory = "6.0.200";
+		static private string DotnetVersionDirectory = "6.0.302";
 			
 		static private string FindRelativeDotnetDirectory(RuntimePlatform.Type HostPlatform)
 		{
