Remove Lumin references in InstalledEngineFilters.xml

diff --git a/Engine/Build/InstalledEngineFilters.xml b/Engine/Build/InstalledEngineFilters.xml
--- a/Engine/Build/InstalledEngineFilters.xml
+++ b/Engine/Build/InstalledEngineFilters.xml
@@ -20,7 +20,6 @@
 	<Property Name="IOSStripFilter" Value="*.a;"/>
 	<Property Name="TVOSStripFilter" Value="*.a;"/>
 	<Property Name="LinuxStripFilter" Value="Engine\Binaries\Linux\*."/>
-	<Property Name="LuminStripFilter" Value="*.a;*.so"/>
 	<Property Name="HoloLensStripFilter" Value="*.pdb"/>
 
 
@@ -549,14 +548,4 @@
 
 	<!-- Win64 specific -->
 
-	<!-- Lumin -->
-	<Property Name="CopyLuminFilter">
-		<!-- Build Files -->
-		Engine/Build/Lumin/...
-		Engine/Source/ThirdParty/ResonanceAudioApi/...
-	</Property>
-	
-	<Property Name="CopyLuminExceptions">
-	</Property>
-
 </BuildGraph>
