Ensure that the UnrealObjectPtrTool executable is built and included in the released editor build for Win64, Mac, and Linux.

diff --git a/Engine/Build/InstalledEngineBuild.xml b/Engine/Build/InstalledEngineBuild.xml
--- a/Engine/Build/InstalledEngineBuild.xml
+++ b/Engine/Build/InstalledEngineBuild.xml
@@ -477,6 +477,7 @@
 			<Compile Target="BootstrapPackagedGame" Configuration="Shipping" Platform="Win64" Tag="#Build Tools Win64"/>
 			<Compile Target="BuildPatchTool" Configuration="Shipping" Platform="Win64" Tag="#Build Tools Win64"/>
 			<Compile Target="SwitchboardListener" Configuration="Development" Platform="Win64" Tag="#Build Tools Win64"/>
+			<Compile Target="UnrealObjectPtrTool" Configuration="Development" Platform="Win64" Tag="#Build Tools Win64"/>
 
 			<Tag Files="Engine/Binaries/Win64/XGEControlWorker.exe" With="#Build Tools Win64"/>
 			<Tag Files="$(RootDir)/Engine/Saved/UnrealBuildTool/BuildConfiguration.Schema.xsd" With="#Build Tools Win64"/>
@@ -543,6 +544,7 @@
 			<Compile Target="BuildPatchTool" Configuration="Shipping" Platform="Linux" Tag="#Build Tools Linux"/>
 			<Compile Target="UnrealRecoverySvc" Configuration="Development" Platform="Linux" Tag="#Build Tools Linux"/>
 			<Compile Target="SwitchboardListener" Configuration="Development" Platform="Linux" Tag="#Build Tools Linux"/>
+			<Compile Target="UnrealObjectPtrTool" Configuration="Development" Platform="Linux" Tag="#Build Tools Linux"/>
 		</Node>
 	</Agent>
 
@@ -605,6 +607,7 @@
 			<Compile Target="DsymExporter" Configuration="Development" Platform="Mac" Tag="#Build Tools Mac"/>
 			<Compile Target="UnrealEditorServices" Configuration="Development" Platform="Mac" Tag="#Build Tools Mac"/>
 			<Compile Target="UnrealAtoS" Configuration="Development" Platform="Mac" Tag="#Build Tools Mac"/>
+			<Compile Target="UnrealObjectPtrTool" Configuration="Development" Platform="Mac" Tag="#Build Tools Mac"/>
 		</Node>
 	</Agent>
 
@@ -620,6 +623,7 @@
 			<Compile Target="BuildPatchTool" Configuration="Shipping" Platform="Linux" Tag="#Build Tools LinuxNative"/>
 			<Compile Target="UnrealRecoverySvc" Configuration="Development" Platform="Linux" Tag="#Build Tools LinuxNative"/>
 			<Compile Target="SwitchboardListener" Configuration="Development" Platform="Linux" Tag="#Build Tools LinuxNative"/>
+			<Compile Target="UnrealObjectPtrTool" Configuration="Development" Platform="Linux" Tag="#Build Tools LinuxNative"/>
 		</Node>
 
 		<!-- Compile the C# tools. Copies files to a temporary directory first, to strip out all confidential folders and avoid sharing violations with the running version of UAT. -->
