Add missing TPS and UPL for OculusOpenXRLoader to binary install for Android

#android

diff --git a/Engine/Build/InstalledEngineFilters.xml b/Engine/Build/InstalledEngineFilters.xml
--- a/Engine/Build/InstalledEngineFilters.xml
+++ b/Engine/Build/InstalledEngineFilters.xml
@@ -466,6 +466,8 @@
 		Engine/Source/ThirdParty/Oculus/OVRPlugin/OVRPlugin/ExtLibs/*.jar
 		Engine/Source/ThirdParty/Oculus/OVRPlugin/*.tps
 		Engine/Source/ThirdParty/Oculus/OVRPlugin/*.xml
+		Engine/Source/ThirdParty/Oculus/OculusOpenXRLoader/*.tps
+		Engine/Source/ThirdParty/Oculus/OculusOpenXRLoader/OculusOpenXRLoader/*.xml
 		Engine/Source/ThirdParty/Oculus/OculusOpenXRLoader/OculusOpenXRLoader/Lib/*/*.so
 		Engine/Source/ThirdParty/ResonanceAudioApi/...
 		Engine/Source/ThirdParty/TangoSDK/...
