Removing Hololens template from InstalledEngineFilters so it can be removed from the engine in 5.1

diff --git a/Engine/Build/InstalledEngineFilters.xml b/Engine/Build/InstalledEngineFilters.xml
--- a/Engine/Build/InstalledEngineFilters.xml
+++ b/Engine/Build/InstalledEngineFilters.xml
@@ -276,7 +276,6 @@
 		Templates/TP_InCamVFXBP/...
 
 		<!-- Enterprise Templates Win64 Only -->
-		Templates/TP_HololensBP/...
 		Templates/TP_DMXBP/...
 
 		<!-- DotNet -->
