- Fix EOSSDK .aar not being included in installer builds since static/dynamic libc++ split.

diff --git a/Engine/Build/InstalledEngineFilters.xml b/Engine/Build/InstalledEngineFilters.xml
--- a/Engine/Build/InstalledEngineFilters.xml
+++ b/Engine/Build/InstalledEngineFilters.xml
@@ -446,7 +446,7 @@
 		Engine/Source/ThirdParty/Android/...
 		Engine/Source/ThirdParty/AndroidPermission/...
 		Engine/Source/ThirdParty/EOSSDK/*.xml
-		Engine/Source/ThirdParty/EOSSDK/SDK/Bin/Android/aar/*.aar
+		Engine/Source/ThirdParty/EOSSDK/....aar
 		Engine/Source/ThirdParty/GoogleARCore/...
 		Engine/Source/ThirdParty/GoogleGameSDK/...
 		Engine/Source/ThirdParty/heapprofd/...
