Adding new Simulation industry blank template to build list. This should expose the new section for simulation in the binary version of 5.1 new project creation.

diff --git a/Engine/Build/InstalledEngineFilters.xml b/Engine/Build/InstalledEngineFilters.xml
--- a/Engine/Build/InstalledEngineFilters.xml
+++ b/Engine/Build/InstalledEngineFilters.xml
@@ -65,6 +65,7 @@
 		Templates/TP_TopDownBP/TP_TopDownBP.uproject
 		Templates/TP_VehicleAdvBP/TP_VehicleAdvBP.uproject
 		Templates/TP_VirtualRealityBP/TP_VirtualRealityBP.uproject
+		Templates/TP_SIM_BlankBP/TP_SIM_BlankBP.uproject
 		Samples/StarterContent/StarterContent.uproject
 	</Property>
 	<!-- Win64 specific DDC -->
@@ -131,6 +132,9 @@
 		Templates/TP_AEC_HandheldARBP/...
 		Templates/TP_MFG_HandheldARBP/...
 
+		Templates/TP_SIM_Blank/...
+		Templates/TP_SIM_BlankBP/...
+
 		<!-- Shared template resources -->
 		Templates/TemplateResources/...
 
