Add UnrealMultiUserSlateServer to the installed build list.

diff --git a/Engine/Build/InstalledEngineBuild.xml b/Engine/Build/InstalledEngineBuild.xml
--- a/Engine/Build/InstalledEngineBuild.xml
+++ b/Engine/Build/InstalledEngineBuild.xml
@@ -472,6 +472,7 @@
 			<Compile Target="UnrealLightmass" Configuration="Development" Platform="Win64" Tag="#Build Tools Win64"/>
 			<Compile Target="UnrealPak" Configuration="Development" Platform="Win64" Tag="#Build Tools Win64"/>
 			<Compile Target="UnrealMultiUserServer" Configuration="Development" Platform="Win64" Tag="#Build Tools Win64"/>
+			<Compile Target="UnrealMultiUserSlateServer" Configuration="Development" Platform="Win64" Tag="#Build Tools Win64"/>
 			<Compile Target="UnrealRecoverySvc" Configuration="Development" Platform="Win64" Tag="#Build Tools Win64"/>
 			<Compile Target="LiveCodingConsole" Configuration="Development" Platform="Win64" Tag="#Build Tools Win64"/>
 			<Compile Target="BootstrapPackagedGame" Configuration="Shipping" Platform="Win64" Tag="#Build Tools Win64"/>
@@ -541,6 +542,7 @@
 			<Compile Target="UnrealLightmass" Configuration="Development" Platform="Linux" Tag="#Build Tools Linux"/>
 			<Compile Target="UnrealPak" Configuration="Development" Platform="Linux" Tag="#Build Tools Linux"/>
 			<Compile Target="UnrealMultiUserServer" Configuration="Development" Platform="Linux" Tag="#Build Tools Linux"/>
+			<Compile Target="UnrealMultiUserSlateServer" Configuration="Development" Platform="Linux" Tag="#Build Tools Linux"/>
 			<Compile Target="BuildPatchTool" Configuration="Shipping" Platform="Linux" Tag="#Build Tools Linux"/>
 			<Compile Target="UnrealRecoverySvc" Configuration="Development" Platform="Linux" Tag="#Build Tools Linux"/>
 			<Compile Target="SwitchboardListener" Configuration="Development" Platform="Linux" Tag="#Build Tools Linux"/>
@@ -602,6 +604,7 @@
 			<Compile Target="UnrealLightmass" Configuration="Development" Platform="Mac" Tag="#Build Tools Mac"/>
 			<Compile Target="UnrealPak" Configuration="Development" Platform="Mac" Tag="#Build Tools Mac"/>
 			<Compile Target="UnrealMultiUserServer" Configuration="Development" Platform="Mac" Tag="#Build Tools Mac"/>
+			<Compile Target="UnrealMultiUserSlateServer" Configuration="Development" Platform="Mac" Tag="#Build Tools Mac"/>
 			<Compile Target="UnrealRecoverySvc" Configuration="Development" Platform="Mac" Tag="#Build Tools Mac"/>
 			<Compile Target="BuildPatchTool" Configuration="Shipping" Platform="Mac" Tag="#Build Tools Mac"/>
 			<Compile Target="DsymExporter" Configuration="Development" Platform="Mac" Tag="#Build Tools Mac"/>
@@ -620,6 +623,7 @@
 			<Compile Target="UnrealLightmass" Configuration="Development" Platform="Linux" Tag="#Build Tools LinuxNative"/>
 			<Compile Target="UnrealPak" Configuration="Development" Platform="Linux" Tag="#Build Tools LinuxNative"/>
 			<Compile Target="UnrealMultiUserServer" Configuration="Development" Platform="Linux" Tag="#Build Tools LinuxNative"/>
+			<Compile Target="UnrealMultiUserSlateServer" Configuration="Development" Platform="Linux" Tag="#Build Tools LinuxNative"/>
 			<Compile Target="BuildPatchTool" Configuration="Shipping" Platform="Linux" Tag="#Build Tools LinuxNative"/>
 			<Compile Target="UnrealRecoverySvc" Configuration="Development" Platform="Linux" Tag="#Build Tools LinuxNative"/>
 			<Compile Target="SwitchboardListener" Configuration="Development" Platform="Linux" Tag="#Build Tools LinuxNative"/>
