Add Options to run additional editor tests with Zen Loader and async thread enabled as part of the editor perf tests

diff --git a/Engine/Build/EditorPerf.xml b/Engine/Build/EditorPerf.xml
--- a/Engine/Build/EditorPerf.xml
+++ b/Engine/Build/EditorPerf.xml
@@ -67,6 +67,16 @@
 	<!-- Iterations to run and average together -->
 	<Option Name="Iterations" DefaultValue="1" Description="The list of iteration numbers to run. For example '1;2;3' for three iterations"/>
 
+	<!-- Arguments for extra Zen Loader tests -->
+	<Option Name="EnableZenLoaderTests" DefaultValue="false" Description="Enable additional test per phase that uses the Zen Loader."/>
+	<Option Name="EnableZenAsyncThread" DefaultValue="true" Description="Enable additional test per phase that uses the Zen Loader."/>
+	<Option Name="EnableZenWithoutEditorDomain" DefaultValue="true" Description="Enable additional test per phase that uses the Zen Loader."/>
+
+	<!-- Zen Loader tests specifc arguments -->
+	<Property Name="ZenLoaderArgs" Value="-zenloader"/>
+	<Property Name="ZenLoaderArgs" Value="$(ZenLoaderArgs),-asyncloadingthread" If="$(EnableZenAsyncThread)" />
+	<Property Name="ZenLoaderArgs" Value="$(ZenLoaderArgs),-noeditordomain" If="$(EnableZenWithoutEditorDomain)" />
+
 	<!-- Gauntlet arguments for all runs -->
 	<Property Name="GauntletArgs" Value="-build=editor -UseEditor -verbose -tempdir=&quot;$(RootDir)/Tests&quot;"/>
 	<Property Name="GauntletArgs" Value="$(GauntletArgs) -branch=$(Branch) -ECBranch=&quot;$(Branch)&quot; -ECChangelist=&quot;$(Change)&quot;"/>
@@ -91,22 +101,27 @@
 	<Property Name="TelemetryConfig" Value="UETelemetry" If="$(IsBuildMachine) And '$(TelemetryConfig)' == ''"/>
 
 	<!-- The main test running macro logic that is expanded with platform and dependency node names -->
-	<Macro Name="RunTest" Arguments="OutputName;Platform;Project;NodeName;Phase;Iteration;Args">
+	<Macro Name="RunTest" Arguments="OutputName;Platform;Project;NodeName;Phase;Iteration;Args;WithZen">
 
 		<Property Name="TestName" Value="EditorPerf$(Phase)"/>
+		<Property Name="TestName" Value="$(TestName)Zen" If="$(WithZen)"/>
 		<Property Name="IterationName" Value="$(Project)$(Platform)$(TestName)$(Iteration)"/>
 		<Property Name="TraceFilePath" Value="$(TraceDir)/$(IterationName).utrace"/>
 		<Property Name="TelemetryFileName" Value="$(IterationName)Telemetry.csv"/>
 		<Property Name="TelemetryFilePath" Value="$(TraceDir)/$(TelemetryFileName)"/>
 		<Property Name="LogOutputParams" Value="-logdir=&quot;$(AutomationOutputDir)/$(OutputName)&quot;"/>
 		<Property Name="ReportParams" Value="-WriteTestResultsForHorde -HordeTestDataKey=&quot;$(NodeName)$(PreflightSuffix)&quot;"/>
-		<Property Name="AllGauntletArgs" Value="-Project=$(Project) -Platform=$(Platform) $(GauntletArgs) $(ReportParams) $(LogOutputParams) -artifactname=&quot;$(Phase)$(Iteration)&quot;"/>
+		<Property Name="ArtifactName" Value="$(Phase)$(Iteration)"/>
+		<Property Name="ArtifactName" Value="$(Phase)$(Iteration)" If="$(WithZen)"/>
+		<Property Name="AllGauntletArgs" Value="-Project=$(Project) -Platform=$(Platform) $(GauntletArgs) $(ReportParams) $(LogOutputParams) -artifactname=&quot;$(ArtifactName)&quot;"/>
+		<Property Name="RunArgs" Value="$(Args)"/>
+		<Property Name="RunArgs" Value="$(Args) -editorargs=&quot;$(ZenLoaderArgs)&quot;" If="$(WithZen)"/>
 
 		<!-- Do DDC cleanup outside the editor process, handles Cold cleanup without changing stats -->
 		<Commandlet Project="$(Project)" Name="DDCCleanup" Arguments="$(Args) $(CommandletArgs)"/>
 
 		<!-- Run the Gauntlet EditorPerf Tests -->
-		<Command Name="RunUnreal" Arguments="$(AllGauntletArgs) $(ProfileArgs) $(Args) -sessionlabel=&quot;$(TestName)&quot; -tracefile=&quot;$(TraceFilePath)&quot;"/>
+		<Command Name="RunUnreal" Arguments="$(AllGauntletArgs) $(ProfileArgs) $(RunArgs) -sessionlabel=&quot;$(TestName)&quot; -tracefile=&quot;$(TraceFilePath)&quot;"/>
 
 		<!-- Translate the utrace file into csv data -->
 		<Commandlet Project="$(Project)" Name="SummarizeTrace" Arguments="-inputfile=&quot;$(TraceFilePath)&quot; -testname=&quot;$(TestName)&quot; $(SummarizeArgs) $(CommandletArgs)"/>
@@ -177,11 +192,16 @@
 						<Do If="'$(Phase)' != 'Incr' or '$(Iteration)' == '1'">
 							<Property Name="RunNodeName" Value="Run $(OutputName) $(Phase)$(Iteration)"/>
 							<Node Name="$(RunNodeName)" Requires="$(RunNodeRequires);$(RunNodeList)">
-								<Expand Name="RunTest" OutputName="$(OutputName)" Platform="$(Platform)" Project="$(Project)" NodeName="$(RunNodeName)" Phase="$(Phase)" Iteration="$(Iteration)" Args="$($(Phase)Args)"/>
+								<Expand Name="RunTest" OutputName="$(OutputName)" Platform="$(Platform)" Project="$(Project)" NodeName="$(RunNodeName)" Phase="$(Phase)" Iteration="$(Iteration)" Args="$($(Phase)Args)" WithZen="false"/>
 							</Node>
 							<Property Name="RunNodeList" Value="$(RunNodeList);$(RunNodeName)"/>
 						</Do>
 					</ForEach>
+					<Property Name="RunNodeZenName" Value="Run $(OutputName) $(Phase) Zen"/>
+					<Node Name="$(RunNodeZenName)" Requires="$(RunNodeRequires);$(RunNodeList)" If="$(EnableZenLoaderTests)">
+						<Expand Name="RunTest" OutputName="$(OutputName)Zen" Platform="$(Platform)" Project="$(Project)" NodeName="$(RunNodeZenName)" Phase="$(Phase)" Iteration="1" Args="$($(Phase)Args)" WithZen="true"/>
+					</Node>
+					<Property Name="RunNodeList" Value="$(RunNodeList);$(RunNodeZenName)" If="$(EnableZenLoaderTests)"/>
 				</ForEach>
 			</Agent>
 			<Label Category="$(RunCategory)" Name="$(Project) Editor Run" Requires="$(RunNodeRequires);$(RunNodeList)" UgsBadge="$(Project) Editor Run"/>
diff --git a/Engine/Source/Editor/UnrealEd/Private/Commandlets/SummarizeTraceCommandlet.cpp b/Engine/Source/Editor/UnrealEd/Private/Commandlets/SummarizeTraceCommandlet.cpp
--- a/Engine/Source/Editor/UnrealEd/Private/Commandlets/SummarizeTraceCommandlet.cpp
+++ b/Engine/Source/Editor/UnrealEd/Private/Commandlets/SummarizeTraceCommandlet.cpp
@@ -2160,6 +2160,7 @@ bool USummarizeTraceCommandlet::GenerateTelemetryCSV(const FString& TestName,
 		{
 			UE_LOG(LogSummarizeTrace, Display, TEXT("Loading project statistics from %s"), *ProjectStatisticsFileName);
 			bool bCSVOk = StatisticDefinition::LoadFromCSV(ProjectStatisticsFileName, NameToDefinitionMap, CpuScopeNamesWithWildcards);
+			check(bCSVOk);
 		}
 	}
 
