Restore performing an enumerate compile for UAT to collect reference libraries

Mac & Linux Only

#rnx

diff --git a/Engine/Build/InstalledEngineBuild.xml b/Engine/Build/InstalledEngineBuild.xml
--- a/Engine/Build/InstalledEngineBuild.xml
+++ b/Engine/Build/InstalledEngineBuild.xml
@@ -559,6 +559,7 @@
 
 			<!-- Copy Source and referenced libraries to a new location with Confidential folders removed-->
 			<Tag Files="$(AutomationToolPaths)" Filter="*.csproj" Except="$(ConfidentialExceptions)"  With="#UAT Projects Mac"/>
+			<CsCompile Project="#UAT Projects Mac" Configuration="Development" Platform="AnyCPU" TagReferences="#UATReferencesMac" EnumerateOnly="true"/>
 			<Tag Files="Engine/Binaries/DotNET/...;Engine/Binaries/ThirdParty/Newtonsoft/...;Engine/Binaries/ThirdParty/IOS/...;Engine/Binaries/ThirdParty/VisualStudio/...;Engine/Source/Programs/...;Engine/Platforms/*/Source/Programs/...;Engine/Source/Editor/SwarmInterface/...;Engine/Intermediate/ScriptModules/...;#UATReferencesMac" Except="$(ConfidentialExceptions)" With="#RedistUATSourceMac"/>
 			<Expand Name="Platform_TagCsToolsRedistUATSource"/>
 			<Copy Files="#RedistUATSourceMac" From="$(RootDir)" To="$(CsToolsDir)"/>
@@ -636,6 +637,7 @@
 
 			<!-- Copy Source and referenced libraries to a new location with Confidential folders removed-->
 			<Tag Files="$(AutomationToolPaths)" Filter="*.csproj" Except="$(ConfidentialExceptions)"  With="#UAT Projects"/>
+			<CsCompile Project="#UAT Projects" Configuration="Development" Platform="AnyCPU" TagReferences="#UATReferences" EnumerateOnly="true"/>
 			<Tag Files="Engine/Binaries/DotNET/...;Engine/Binaries/ThirdParty/Newtonsoft/...;Engine/Binaries/ThirdParty/IOS/...;Engine/Binaries/ThirdParty/VisualStudio/...;Engine/Source/Programs/...;Engine/Platforms/*/Source/Programs/...;Engine/Source/Editor/SwarmInterface/...;Engine/Intermediate/ScriptModules/...;#UATReferences" Except="$(ConfidentialExceptions)" With="#RedistUATSource"/>
 			<Expand Name="Platform_TagCsToolsRedistUATSource"/>
 			<Copy Files="#RedistUATSource" From="$(RootDir)" To="$(CsToolsDir)"/>
