Fixing up some UAT/UBT issues with installed builds

* Missing various dlls
* Excluding obj folders for Automation tool and platform folders
* Full fix for mac missing dlls

#rnx

diff --git a/Engine/Build/InstalledEngineBuild.xml b/Engine/Build/InstalledEngineBuild.xml
--- a/Engine/Build/InstalledEngineBuild.xml
+++ b/Engine/Build/InstalledEngineBuild.xml
@@ -498,6 +498,7 @@
 
 			<!-- Copy Source and referenced libraries to a new location with Confidential folders removed-->
 			<Tag Files="$(AutomationToolPaths)" Filter="*.csproj" Except="$(ConfidentialExceptions)"  With="#UAT Projects"/>
+			<CsCompile Project="#UAT Projects" Configuration="Development" Platform="AnyCPU" TagReferences="#UATReferences" EnumerateOnly="true"/>
 			<Tag Files="Engine/Binaries/DotNET/...;Engine/Binaries/ThirdParty/Newtonsoft/...;Engine/Binaries/ThirdParty/IOS/...;Engine/Binaries/ThirdParty/VisualStudio/...;Engine/Source/Programs/...;Engine/Platforms/*/Source/Programs/...;Engine/Source/Editor/SwarmInterface/...;Engine/Intermediate/ScriptModules/...;#UATReferences" Except="$(ConfidentialExceptions)" With="#RedistUATSource"/>
 			<Expand Name="Platform_TagCsToolsRedistUATSource"/>
 			<Copy Files="#RedistUATSource" From="$(RootDir)" To="$(CsToolsDir)"/>
@@ -509,6 +510,7 @@
 			<Tag Files="$(RedistAutomationToolPaths)" Filter="*.csproj" With="#RedistUATProjects"/>
 
 			<!-- Compile all the tools -->
+			<CsCompile Project="#RedistUATProjects" Configuration="Development" Platform="AnyCPU" Tag="#Build Tools CS Binaries;#Sign Binaries" TagReferences="#Build Tools CS Binaries"/>
 			<Do If="'$(HostPlatform)' == 'Win64'">
 				<CsCompile Project="$(CsToolsDir)/Engine/Source/Programs/UnrealSwarm/SwarmCoordinator/SwarmCoordinator.csproj" Configuration="Development" Platform="AnyCPU" UseSystemCompiler="true" Tag="#Build Tools CS Binaries;#Sign Binaries" TagReferences="#Build Tools CS Binaries"/>
 				<CsCompile Project="$(CsToolsDir)/Engine/Source/Programs/UnrealSwarm/Agent/Agent.csproj" Configuration="Development" Platform="AnyCPU" UseSystemCompiler="true" Tag="#Build Tools CS Binaries;#Sign Binaries" TagReferences="#Build Tools CS Binaries"/>
@@ -525,6 +527,9 @@
 			<!-- Tag AutomationTool Script module build records, so that prebuilt modules may be discovered in the absence of source code -->
 			<Tag Files="$(CsToolsDir)/Engine/Intermediate/ScriptModules/..." With="#Build Tools CS Binaries" />
 
+			<!-- Tag the native P4 libraries as they are not a C# project so not a reference project -->
+			<Tag Files="$(CsToolsDir)/Engine/Binaries/DotNET/EpicGames.Perforce.Native/..." With="#Build Tools CS Binaries" />
+
 			<CsCompile Project="$(CsToolsDir)/Engine/Source/Programs/NetworkProfiler/NetworkProfiler/NetworkProfiler.csproj" Configuration="Development" Platform="AnyCPU" UseSystemCompiler="true" Tag="#Build Tools CS Binaries" TagReferences="#Build Tools CS Binaries;#Sign Binaries" If="'$(HostPlatform)' != 'Linux'"/>
 			<CsCompile Project="$(CsToolsDir)/Engine/Source/Programs/IOS/MobileDeviceInterface/MobileDeviceInterface.csproj" Configuration="Development" Platform="AnyCPU" UseSystemCompiler="true" Arguments="/verbosity:minimal /target:Rebuild" Tag="#Build Tools CS Binaries;#Sign Binaries" TagReferences="#Build Tools CS Binaries"/>
 
@@ -570,12 +575,17 @@
 			</ForEach>
 			<Tag Files="$(RedistAutomationToolPaths)" Filter="*.csproj" With="#RedistUATProjectsMac"/>
 
+			<CsCompile Project="#RedistUATProjectsMac" Configuration="Development" Platform="AnyCPU" Tag="#Build Tools CS Binaries Mac;#Sign Binaries Mac" TagReferences="#Build Tools CS Binaries Mac"/>
+
 			<!-- Tag AutomationTool and UnrealBuildTool folders recursively as NET Core dependencies are not currently handled by CsCompile -->
 			<Tag Files="$(CsToolsDir)/Engine/Binaries/DotNET/AutomationTool/..." With="#Build Tools CS Binaries Mac" />
 			<Tag Files="$(CsToolsDir)/Engine/Binaries/DotNET/UnrealBuildTool/..." With="#Build Tools CS Binaries Mac" />
 
 			<!-- Tag AutomationTool Script module build records, so that prebuilt modules may be discovered in the absence of source code -->
 			<Tag Files="$(CsToolsDir)/Engine/Intermediate/ScriptModules/..." With="#Build Tools CS Binaries Mac" />
+
+			<!-- Tag the native P4 libraries as they are not a C# project so not a reference project -->
+			<Tag Files="$(CsToolsDir)/Engine/Binaries/DotNET/EpicGames.Perforce.Native/..." With="#Build Tools CS Binaries Mac" />
 						
 			<CsCompile Project="$(CsToolsDir)/Engine/Source/Programs/NetworkProfiler/NetworkProfiler/NetworkProfiler.csproj" Configuration="Development" Platform="AnyCPU" UseSystemCompiler="true" Tag="#Build Tools CS Binaries Mac" TagReferences="#Build Tools CS Binaries Mac;#Sign Binaries Mac"/>
 			<CsCompile Project="$(CsToolsDir)/Engine/Source/Programs/IOS/MobileDeviceInterface/MobileDeviceInterface.csproj" Configuration="Development" Platform="AnyCPU" UseSystemCompiler="true" Arguments="/verbosity:minimal /target:Rebuild" Tag="#Build Tools CS Binaries Mac;#Sign Binaries Mac" TagReferences="#Build Tools CS Binaries Mac"/>
@@ -647,11 +657,15 @@
 				<Property Name="RedistAutomationToolPaths" Value="$(RedistAutomationToolPaths)$(CsToolsDir)/$(AutomationToolPath);"/>
 			</ForEach>
 			<Tag Files="$(RedistAutomationToolPaths)" Filter="*.csproj" With="#RedistUATProjects"/>
+			<CsCompile Project="#RedistUATProjects" Configuration="Development" Platform="AnyCPU" Tag="#Build Tools CS Binaries LinuxNative;#Sign Binaries LinuxNative" TagReferences="#Build Tools CS Binaries LinuxNative"/>
 
 			<!-- Tag AutomationTool and UnrealBuildTool folders recursively as NET Core dependencies are not currently handled by CsCompile -->
 			<Tag Files="$(CsToolsDir)/Engine/Binaries/DotNET/AutomationTool/..." With="#Build Tools CS Binaries LinuxNative" />
 			<Tag Files="$(CsToolsDir)/Engine/Binaries/DotNET/UnrealBuildTool/..." With="#Build Tools CS Binaries LinuxNative" />
 
+			<!-- Tag the native P4 libraries as they are not a C# project so not a reference project -->
+			<Tag Files="$(CsToolsDir)/Engine/Binaries/DotNET/EpicGames.Perforce.Native/..." With="#Build Tools CS LinuxNative" />
+
 			<!-- Tag AutomationTool Script module build records, so that prebuilt modules may be discovered in the absence of source code -->
 			<Tag Files="$(CsToolsDir)/Engine/Intermediate/ScriptModules/..." With="#Build Tools CS Binaries LinuxNative" />
 
diff --git a/Engine/Build/InstalledEngineFilters.xml b/Engine/Build/InstalledEngineFilters.xml
--- a/Engine/Build/InstalledEngineFilters.xml
+++ b/Engine/Build/InstalledEngineFilters.xml
@@ -421,6 +421,11 @@
 		Engine/Source/ThirdParty/EOSSDK/SDK/Bin/x64/xaudio2_9redist.dll
 	</Property>
 	<Property Name="CopyWin64Exceptions">
+		Engine/Platforms/*/Source/Programs/AutomationTool/obj/...
+		Engine/Platforms/*/Source/Programs/AutomationTool/*/obj/...	
+		Engine/Platforms/*/Source/Programs/UnrealBuildTool/obj/...
+		Engine/Source/Programs/AutomationTool/obj/...
+		Engine/Source/Programs/AutomationTool/*/obj/...
 		Engine/Source/Programs/UnrealBuildTool/obj/...
 		Engine/Source/Programs/Shared/*/bin/...
 		Engine/Source/Programs/Shared/*/obj/...
@@ -435,6 +440,11 @@
 		Engine/Source/ThirdParty/PLCrashReporter/lib/lib-Xcode-12.4/Mac/...
 	</Property>
 	<Property Name="CopyMacExceptions">
+		Engine/Platforms/*/Source/Programs/AutomationTool/obj/...
+		Engine/Platforms/*/Source/Programs/AutomationTool/*/obj/...
+		Engine/Platforms/*/Source/Programs/UnrealBuildTool/obj/...
+		Engine/Source/Programs/AutomationTool/obj/...
+		Engine/Source/Programs/AutomationTool/*/obj/...
 		Engine/Source/Programs/UnrealBuildTool/obj/...
 		Engine/Source/Programs/Shared/*/bin/...
 		Engine/Source/Programs/Shared/*/obj/...
@@ -547,6 +557,14 @@
 	</Property>
 
 	<Property Name="CopyLinuxExceptions">
+		Engine/Platforms/*/Source/Programs/AutomationTool/obj/...
+		Engine/Platforms/*/Source/Programs/AutomationTool/*/obj/...
+		Engine/Platforms/*/Source/Programs/UnrealBuildTool/obj/...
+		Engine/Source/Programs/AutomationTool/obj/...
+		Engine/Source/Programs/AutomationTool/*/obj/...
+		Engine/Source/Programs/UnrealBuildTool/obj/...
+		Engine/Source/Programs/Shared/*/bin/...
+		Engine/Source/Programs/Shared/*/obj/...
 	</Property>
 
 	<Property Name="CopyLinuxFilter" Value="$(CopyLinuxFilter);Engine/Extras/ThirdPartyNotUE/putty/..."/>	
