- Change Mac to use "dotnet msbuild" to compile C# programs with, instead of msbuild - dotnet is shipped with the engine so is always there. This is to fix builders building with xbuild which is deprecated and going away
- Removed unused MobileDeviceInterface

diff --git a/Engine/Build/InstalledEngineBuild.xml b/Engine/Build/InstalledEngineBuild.xml
--- a/Engine/Build/InstalledEngineBuild.xml
+++ b/Engine/Build/InstalledEngineBuild.xml
@@ -517,6 +517,7 @@
 				<CsCompile Project="$(CsToolsDir)/Engine/Source/Editor/SwarmInterface/DotNET/SwarmInterface.csproj" Configuration="Development" Platform="AnyCPU" UseSystemCompiler="true" Tag="#Build Tools CS Binaries;#Sign Binaries" TagReferences="#Build Tools CS Binaries"/>
 				<CsCompile Project="$(CsToolsDir)/Engine/Source/Programs/UnrealControls/UnrealControls.csproj" Configuration="Development" Platform="AnyCPU" UseSystemCompiler="true" Tag="#Build Tools CS Binaries;#Sign Binaries" TagReferences="#Build Tools CS Binaries"/>
 				<CsCompile Project="$(CsToolsDir)/Engine/Source/Programs/IOS/iPhonePackager/iPhonePackager.csproj" Configuration="Development" Platform="AnyCPU" UseSystemCompiler="true" Arguments="/verbosity:minimal /target:Rebuild" Tag="#Build Tools CS Binaries;#Sign Binaries" TagReferences="#Build Tools CS Binaries"/>
+				<CsCompile Project="$(CsToolsDir)/Engine/Source/Programs/NetworkProfiler/NetworkProfiler/NetworkProfiler.csproj" Configuration="Development" Platform="AnyCPU" UseSystemCompiler="true" Tag="#Build Tools CS Binaries" TagReferences="#Build Tools CS Binaries;#Sign Binaries"/>
 				<Expand Name="Platform_Build_Tools_CS_Binaries"/>
 			</Do>
 			
@@ -530,9 +531,6 @@
 			<!-- Tag the native P4 libraries as they are not a C# project so not a reference project -->
 			<Tag Files="$(CsToolsDir)/Engine/Binaries/DotNET/EpicGames.Perforce.Native/..." With="#Build Tools CS Binaries" />
 
-			<CsCompile Project="$(CsToolsDir)/Engine/Source/Programs/NetworkProfiler/NetworkProfiler/NetworkProfiler.csproj" Configuration="Development" Platform="AnyCPU" UseSystemCompiler="true" Tag="#Build Tools CS Binaries" TagReferences="#Build Tools CS Binaries;#Sign Binaries" If="'$(HostPlatform)' != 'Linux'"/>
-			<CsCompile Project="$(CsToolsDir)/Engine/Source/Programs/IOS/MobileDeviceInterface/MobileDeviceInterface.csproj" Configuration="Development" Platform="AnyCPU" UseSystemCompiler="true" Arguments="/verbosity:minimal /target:Rebuild" Tag="#Build Tools CS Binaries;#Sign Binaries" TagReferences="#Build Tools CS Binaries"/>
-
 			<!-- Create tag set for saved source needed for documentation -->
 			<Tag Files="$(CsToolsDir)/Engine/..." With="#SourceForDocumentation" />
 
@@ -583,13 +581,10 @@
 
 			<!-- Tag AutomationTool Script module build records, so that prebuilt modules may be discovered in the absence of source code -->
 			<Tag Files="$(CsToolsDir)/Engine/Intermediate/ScriptModules/..." With="#Build Tools CS Binaries Mac" />
-
+						
 			<!-- Tag the native P4 libraries as they are not a C# project so not a reference project -->
 			<Tag Files="$(CsToolsDir)/Engine/Binaries/DotNET/EpicGames.Perforce.Native/..." With="#Build Tools CS Binaries Mac" />
 						
-			<CsCompile Project="$(CsToolsDir)/Engine/Source/Programs/NetworkProfiler/NetworkProfiler/NetworkProfiler.csproj" Configuration="Development" Platform="AnyCPU" UseSystemCompiler="true" Tag="#Build Tools CS Binaries Mac" TagReferences="#Build Tools CS Binaries Mac;#Sign Binaries Mac"/>
-			<CsCompile Project="$(CsToolsDir)/Engine/Source/Programs/IOS/MobileDeviceInterface/MobileDeviceInterface.csproj" Configuration="Development" Platform="AnyCPU" UseSystemCompiler="true" Arguments="/verbosity:minimal /target:Rebuild" Tag="#Build Tools CS Binaries Mac;#Sign Binaries Mac" TagReferences="#Build Tools CS Binaries Mac"/>
-
 			<!-- IPP doesn't currently compile with the bundled Mono. -->
 			<!-- CsCompile Project="$(CsToolsDir)/Engine/Source/Programs/UnrealControls/UnrealControls.csproj" Configuration="Development" Platform="AnyCPU" Tag="#Build Tools CS Binaries Mac;#Sign Binaries Mac" TagReferences="#Build Tools CS Binaries Mac"/-->
 			<!-- CsCompile Project="$(CsToolsDir)/Engine/Source/Programs/IOS/iPhonePackager/iPhonePackager.csproj" Configuration="Development" Platform="AnyCPU" Arguments="/verbosity:minimal /target:Rebuild" Tag="#Build Tools CS Binaries Mac;#Sign Binaries Mac" TagReferences="#Build Tools CS Binaries Mac"/-->
diff --git a/Engine/Source/Programs/AutomationTool/AutomationUtils/MacHostPlatform.cs b/Engine/Source/Programs/AutomationTool/AutomationUtils/MacHostPlatform.cs
--- a/Engine/Source/Programs/AutomationTool/AutomationUtils/MacHostPlatform.cs
+++ b/Engine/Source/Programs/AutomationTool/AutomationUtils/MacHostPlatform.cs
@@ -22,11 +22,13 @@ namespace AutomationTool
 			// As of 5.0 mono comes with msbuild which performs better. If that's installed then use it
 			if (string.IsNullOrEmpty(CachedFrameworkMsbuildExe))
 			{
-				bool CanUseMsBuild = string.IsNullOrEmpty(CommandUtils.WhichApp("msbuild")) == false;
+				bool CanUseMsBuild = string.IsNullOrEmpty(CommandUtils.WhichApp("dotnet")) == false;
 
 				if (CanUseMsBuild)
 				{
-					CachedFrameworkMsbuildExe = "msbuild";
+					Log.TraceInformation($"using {CommandUtils.WhichApp("dotnet")}!");
+
+					CachedFrameworkMsbuildExe = "dotnet msbuild";
 				}
 				else
 				{
diff --git a/Engine/Source/Programs/AutomationTool/AutomationUtils/ProcessUtils.cs b/Engine/Source/Programs/AutomationTool/AutomationUtils/ProcessUtils.cs
--- a/Engine/Source/Programs/AutomationTool/AutomationUtils/ProcessUtils.cs
+++ b/Engine/Source/Programs/AutomationTool/AutomationUtils/ProcessUtils.cs
@@ -837,6 +837,14 @@ namespace AutomationTool
 			// Check if the application exists, including the PATH directories.
 			if (Options.HasFlag(ERunOptions.AppMustExist) && !FileExists(Options.HasFlag(ERunOptions.NoLoggingOfRunCommand) ? true : false, App))
 			{
+				// in the case of something like "dotnet msbuild", split it up and put the msbuild on the commandline
+				// this could be generalized, but would have to account for spaces in the App part
+				if (App.StartsWith("dotnet "))
+				{
+					string[] Tokens = App.Split(" ", StringSplitOptions.RemoveEmptyEntries);
+					App = Tokens[0];
+					CommandLine = $"{Tokens[1]} {CommandLine}";
+				}
 				string ResolvedPath = WhichApp(App);
 
 				if(string.IsNullOrEmpty(ResolvedPath))
diff --git a/Engine/Source/Programs/AutomationTool/BuildGraph/Scripts/InstalledBuild.cs b/Engine/Source/Programs/AutomationTool/BuildGraph/Scripts/InstalledBuild.cs
--- a/Engine/Source/Programs/AutomationTool/BuildGraph/Scripts/InstalledBuild.cs
+++ b/Engine/Source/Programs/AutomationTool/BuildGraph/Scripts/InstalledBuild.cs
@@ -360,15 +360,9 @@ namespace AutomationTool
 				Output += await CsCompileAsync(FileReference.Combine(TempCsToolsDir, "Engine/Source/Editor/SwarmInterface/DotNET/SwarmInterface.csproj"), Platform: "AnyCPU", Configuration: "Development");
 				Output += await CsCompileAsync(FileReference.Combine(TempCsToolsDir, "Engine/Source/Programs/UnrealControls/UnrealControls.csproj"), Platform: "AnyCPU", Configuration: "Development");
 				Output += await CsCompileAsync(FileReference.Combine(TempCsToolsDir, "Engine/Source/Programs/IOS/iPhonePackager/iPhonePackager.csproj"), Platform: "AnyCPU", Configuration: "Development", Arguments: "/verbosity:minimal /target:Rebuild");
-			}
-
-			if (HostPlatform.Current.HostEditorPlatform != UnrealTargetPlatform.Linux)
-			{
 				Output += await CsCompileAsync(FileReference.Combine(TempCsToolsDir, "Engine/Source/Programs/NetworkProfiler/NetworkProfiler/NetworkProfiler.csproj"), Platform: "AnyCPU", Configuration: "Development");
 			}
 
-			Output += await CsCompileAsync(FileReference.Combine(RootDir, "Engine/Source/Programs/IOS/MobileDeviceInterface/MobileDeviceInterface.csproj"), Platform: "AnyCPU", Configuration: "Development", Arguments: "/verbosity:minimal /target:Rebuild");
-
 			FileSet Binaries = Output.Binaries + Output.References;
 
 			// Tag AutomationTool and UnrealBuildTool folders recursively as NET Core dependencies are not currently handled by CsCompile
diff --git a/Engine/Source/Programs/AutomationTool/IOS/IOS.Automation.csproj b/Engine/Source/Programs/AutomationTool/IOS/IOS.Automation.csproj
--- a/Engine/Source/Programs/AutomationTool/IOS/IOS.Automation.csproj
+++ b/Engine/Source/Programs/AutomationTool/IOS/IOS.Automation.csproj
@@ -39,10 +39,6 @@
       <HintPath>..\..\..\..\Binaries\DotNET\Ionic.Zip.Reduced.dll</HintPath>
       <Private>False</Private>
     </Reference>
-    <Reference Include="MobileDeviceInterface">
-      <HintPath>..\..\..\..\Binaries\DotNET\IOS\MobileDeviceInterface.dll</HintPath>
-      <Private>False</Private>
-    </Reference>
   </ItemGroup>
 
   <ItemGroup>
diff --git a/Engine/Source/Programs/AutomationTool/IOS/IOSPlatform.Automation.cs b/Engine/Source/Programs/AutomationTool/IOS/IOSPlatform.Automation.cs
--- a/Engine/Source/Programs/AutomationTool/IOS/IOSPlatform.Automation.cs
+++ b/Engine/Source/Programs/AutomationTool/IOS/IOSPlatform.Automation.cs
@@ -2288,18 +2288,6 @@ public class IOSPlatform : ApplePlatform
 		}
 	}
 
-#region Hooks
-
-	public override void PreBuildAgenda(UnrealBuild Build, UnrealBuild.BuildAgenda Agenda, ProjectParams Params)
-	{
-		if (UnrealBuildTool.BuildHostPlatform.Current.Platform != UnrealTargetPlatform.Mac && !Unreal.IsEngineInstalled())
-		{
-			Agenda.DotNetProjects.Add(@"Engine\Source\Programs\IOS\MobileDeviceInterface\MobileDeviceInterface.csproj");
-		}
-	}
-
-#endregion
-
 	public override DirectoryReference GetProjectRootForStage(DirectoryReference RuntimeRoot, StagedDirectoryReference RelativeProjectRootForStage)
 	{
 		return DirectoryReference.Combine(RuntimeRoot, "cookeddata/" + RelativeProjectRootForStage.Name);
diff --git a/Engine/Source/Programs/AutomationTool/Scripts/BuildCommonTools.Automation.cs b/Engine/Source/Programs/AutomationTool/Scripts/BuildCommonTools.Automation.cs
--- a/Engine/Source/Programs/AutomationTool/Scripts/BuildCommonTools.Automation.cs
+++ b/Engine/Source/Programs/AutomationTool/Scripts/BuildCommonTools.Automation.cs
@@ -123,9 +123,6 @@ public class BuildCommonTools : BuildCommand
 		{
 			Agenda.DotNetProjects.Add(@"Engine/Source/Programs/iOS/iPhonePackager/iPhonePackager.csproj");
 			ExtraBuildProducts.Add(CommandUtils.CombinePaths(CommandUtils.CmdEnv.LocalRoot, @"Engine/Binaries/DotNET/iOS/iPhonePackager.exe"));
-
-			Agenda.DotNetProjects.Add(@"Engine/Source/Programs/iOS/MobileDeviceInterface/MobileDeviceInterface.csproj");
-			ExtraBuildProducts.Add(CommandUtils.CombinePaths(CommandUtils.CmdEnv.LocalRoot, @"Engine/Binaries/DotNET/iOS/MobileDeviceInterface.dll"));
 		}
 
 		// Platform extensions
diff --git a/Engine/Source/Programs/IOS/MobileDeviceInterface/ConnectEventArgs.cs b/Engine/Source/Programs/IOS/MobileDeviceInterface/ConnectEventArgs.cs
deleted file mode 100644
diff --git a/Engine/Source/Programs/IOS/MobileDeviceInterface/ConnectEventHandler.cs b/Engine/Source/Programs/IOS/MobileDeviceInterface/ConnectEventHandler.cs
deleted file mode 100644
diff --git a/Engine/Source/Programs/IOS/MobileDeviceInterface/CoreFoundation.cs b/Engine/Source/Programs/IOS/MobileDeviceInterface/CoreFoundation.cs
deleted file mode 100644
diff --git a/Engine/Source/Programs/IOS/MobileDeviceInterface/DeviceNotificationEventArgs.cs b/Engine/Source/Programs/IOS/MobileDeviceInterface/DeviceNotificationEventArgs.cs
deleted file mode 100644
diff --git a/Engine/Source/Programs/IOS/MobileDeviceInterface/IPhoneFile.cs b/Engine/Source/Programs/IOS/MobileDeviceInterface/IPhoneFile.cs
deleted file mode 100644
diff --git a/Engine/Source/Programs/IOS/MobileDeviceInterface/MobileDevice.cs b/Engine/Source/Programs/IOS/MobileDeviceInterface/MobileDevice.cs
deleted file mode 100644
diff --git a/Engine/Source/Programs/IOS/MobileDeviceInterface/MobileDeviceInstance.cs b/Engine/Source/Programs/IOS/MobileDeviceInterface/MobileDeviceInstance.cs
deleted file mode 100644
diff --git a/Engine/Source/Programs/IOS/MobileDeviceInterface/MobileDeviceInterface.csproj b/Engine/Source/Programs/IOS/MobileDeviceInterface/MobileDeviceInterface.csproj
deleted file mode 100644
diff --git a/Engine/Source/Programs/IOS/MobileDeviceInterface/Properties/AssemblyInfo.cs b/Engine/Source/Programs/IOS/MobileDeviceInterface/Properties/AssemblyInfo.cs
deleted file mode 100644
diff --git a/Engine/Source/Programs/IOS/MobileDeviceInterface/Utilities.cs b/Engine/Source/Programs/IOS/MobileDeviceInterface/Utilities.cs
deleted file mode 100644
diff --git a/Engine/Source/Programs/IOS/MobileDeviceInterface/app.config b/Engine/Source/Programs/IOS/MobileDeviceInterface/app.config
deleted file mode 100644
diff --git a/Engine/Source/Programs/IOS/iPhonePackager/DeployTime.cs b/Engine/Source/Programs/IOS/iPhonePackager/DeployTime.cs
--- a/Engine/Source/Programs/IOS/iPhonePackager/DeployTime.cs
+++ b/Engine/Source/Programs/IOS/iPhonePackager/DeployTime.cs
@@ -95,7 +95,7 @@ namespace iPhonePackager
 					string DeviceName = DeviceInfo.DeviceName;
 
 					// Check the IPA's mobile provision against the connected device to make sure this device is authorized
-					// We'll still try installing anyways, but this message is more friendly than the failure we get back from MobileDeviceInterface
+					// We'll still try installing anyways, but this message is more friendly than the failure we got back from MobileDeviceInterface when we used it
 					if (UDID != String.Empty)
 					{
 						if (!Provision.ContainsUDID(UDID))
