Disable AllowParallelExecutor by default for UnrealGame

Work around for ParallelExecutor issues with PCH's for some machine setups

#rnx

diff --git a/Engine/Build/InstalledEngineBuild.xml b/Engine/Build/InstalledEngineBuild.xml
--- a/Engine/Build/InstalledEngineBuild.xml
+++ b/Engine/Build/InstalledEngineBuild.xml
@@ -69,6 +69,11 @@
 	<!-- Ouput directory for the build -->
 	<Option Name="BuiltDirectory" DefaultValue="$(RootDir)/LocalBuilds/Engine" Description="Directory for outputting the built engine"/>
 
+	<!-- Control ParallelExecutor -->
+	<Option	Name="AllowParallelExecutor" Restrict="true|false" DefaultValue="false" Description="Allow the use of UAT's ParallelExecutor when compiling"/>
+	<Property Name="bAllowParallelExecutor" Value="$(AllowParallelExecutor)" />
+	<Property Name="bAllowParallelExecutor" Value="true" If="$(IsBuildMachine)" />
+
 	<!-- The local output directory -->
 	<Property Name="LocalInstalledDir" Value="$(BuiltDirectory)/Windows"/>
 	<Property Name="LocalInstalledDirMac" Value="$(BuiltDirectory)/Mac"/>
@@ -265,8 +270,8 @@
 				<Tag Files="$(Tag)" With="#UnrealGame Win64"/>
 			</ForEach>
 			<ForEach Name="Configuration" Values="$(GameConfigurations)">
-				<Compile Target="UnrealGame" Platform="Win64" Configuration="$(Configuration)" Tag="#UnrealGame Win64" Arguments="-precompile -allmodules -nolink $(TargetDebugInfoArg)"/>
-				<Compile Target="UnrealGame" Platform="Win64" Configuration="$(Configuration)" Tag="#UnrealGame Win64" Arguments="-precompile $(TargetDebugInfoArg)" Clean="false"/>
+				<Compile Target="UnrealGame" Platform="Win64" Configuration="$(Configuration)" Tag="#UnrealGame Win64" AllowParallelExecutor="$(bAllowParallelExecutor)" Arguments="-precompile -allmodules -nolink $(TargetDebugInfoArg)"/>
+				<Compile Target="UnrealGame" Platform="Win64" Configuration="$(Configuration)" Tag="#UnrealGame Win64" AllowParallelExecutor="$(bAllowParallelExecutor)" Arguments="-precompile $(TargetDebugInfoArg)" Clean="false"/>
 			</ForEach>
 
 			<Do If="$(EmbedSrcSrvInfo)">
@@ -302,8 +307,8 @@
 		<Node Name="Compile UnrealGame Linux" Requires="Compile UnrealHeaderTool Win64" Produces="#UnrealGame Linux;#UnrealGame Linux Unstripped;#UnrealGame Linux Stripped">
 			<ForEach Name="Target" Values="UnrealGame;$(OptionalClientTarget);$(OptionalServerTarget)">
 				<ForEach Name="Configuration" Values="$(GameConfigurations)">
-					<Compile Target="$(Target)" Platform="Linux" Configuration="$(Configuration)" Tag="#UnrealGame Linux" Arguments="-precompile -allmodules -nolink $(TargetDebugInfoArg)"/>
-					<Compile Target="$(Target)" Platform="Linux" Configuration="$(Configuration)" Tag="#UnrealGame Linux" Arguments="-precompile $(TargetDebugInfoArg)" Clean="false"/>
+					<Compile Target="$(Target)" Platform="Linux" Configuration="$(Configuration)" Tag="#UnrealGame Linux" AllowParallelExecutor="$(bAllowParallelExecutor)" Arguments="-precompile -allmodules -nolink $(TargetDebugInfoArg)"/>
+					<Compile Target="$(Target)" Platform="Linux" Configuration="$(Configuration)" Tag="#UnrealGame Linux" AllowParallelExecutor="$(bAllowParallelExecutor)" Arguments="-precompile $(TargetDebugInfoArg)" Clean="false"/>
 				</ForEach>
 			</ForEach>
 
@@ -317,8 +322,8 @@
 		<Node Name="Compile UnrealGame LinuxNative" Requires="Compile UnrealHeaderTool Linux" Produces="#UnrealGame LinuxNative;#UnrealGame LinuxNative Unstripped;#UnrealGame LinuxNative Stripped">
 			<ForEach Name="Target" Values="UnrealGame;$(OptionalClientTarget);$(OptionalServerTarget)">
 				<ForEach Name="Configuration" Values="$(GameConfigurations)">
-					<Compile Target="$(Target)" Platform="Linux" Configuration="$(Configuration)" Tag="#UnrealGame LinuxNative" Arguments="-precompile -allmodules -nolink $(TargetDebugInfoArg)"/>
-					<Compile Target="$(Target)" Platform="Linux" Configuration="$(Configuration)" Tag="#UnrealGame LinuxNative" Arguments="-precompile $(TargetDebugInfoArg)" Clean="false"/>
+					<Compile Target="$(Target)" Platform="Linux" Configuration="$(Configuration)" Tag="#UnrealGame LinuxNative" AllowParallelExecutor="$(bAllowParallelExecutor)" Arguments="-precompile -allmodules -nolink $(TargetDebugInfoArg)"/>
+					<Compile Target="$(Target)" Platform="Linux" Configuration="$(Configuration)" Tag="#UnrealGame LinuxNative" AllowParallelExecutor="$(bAllowParallelExecutor)" Arguments="-precompile $(TargetDebugInfoArg)" Clean="false"/>
 				</ForEach>
 			</ForEach>
 
@@ -332,8 +337,8 @@
 		<Node Name="Compile UnrealGame LinuxArm64" Requires="Compile UnrealHeaderTool Win64" Produces="#UnrealGame LinuxArm64;#UnrealGame LinuxArm64 Unstripped;#UnrealGame LinuxArm64 Stripped">
 			<ForEach Name="Target" Values="UnrealGame;$(OptionalClientTarget);$(OptionalServerTarget)">
 				<ForEach Name="Configuration" Values="$(GameConfigurations)">
-					<Compile Target="$(Target)" Platform="LinuxArm64" Configuration="$(Configuration)" Tag="#UnrealGame LinuxArm64" Arguments="-precompile -allmodules -nolink $(TargetDebugInfoArg)"/>
-					<Compile Target="$(Target)" Platform="LinuxArm64" Configuration="$(Configuration)" Tag="#UnrealGame LinuxArm64" Arguments="-precompile $(TargetDebugInfoArg)" Clean="false"/>
+					<Compile Target="$(Target)" Platform="LinuxArm64" Configuration="$(Configuration)" Tag="#UnrealGame LinuxArm64" AllowParallelExecutor="$(bAllowParallelExecutor)" Arguments="-precompile -allmodules -nolink $(TargetDebugInfoArg)"/>
+					<Compile Target="$(Target)" Platform="LinuxArm64" Configuration="$(Configuration)" Tag="#UnrealGame LinuxArm64" AllowParallelExecutor="$(bAllowParallelExecutor)" Arguments="-precompile $(TargetDebugInfoArg)" Clean="false"/>
 				</ForEach>
 			</ForEach>
 
@@ -347,8 +352,8 @@
 		<Node Name="Compile UnrealGame LinuxArm64 LinuxNative" Requires="Compile UnrealHeaderTool Linux" Produces="#UnrealGame LinuxArm64 LinuxNative;#UnrealGame LinuxArm64 LinuxNative Unstripped;#UnrealGame LinuxArm64 LinuxNative Stripped">
 			<ForEach Name="Target" Values="UnrealGame;$(OptionalClientTarget);$(OptionalServerTarget)">
 				<ForEach Name="Configuration" Values="$(GameConfigurations)">
-					<Compile Target="$(Target)" Platform="LinuxArm64" Configuration="$(Configuration)" Tag="#UnrealGame LinuxArm64 LinuxNative" Arguments="-precompile -allmodules -nolink $(TargetDebugInfoArg)"/>
-					<Compile Target="$(Target)" Platform="LinuxArm64" Configuration="$(Configuration)" Tag="#UnrealGame LinuxArm64 LinuxNative" Arguments="-precompile $(TargetDebugInfoArg)" Clean="false"/>
+					<Compile Target="$(Target)" Platform="LinuxArm64" Configuration="$(Configuration)" Tag="#UnrealGame LinuxArm64 LinuxNative" AllowParallelExecutor="$(bAllowParallelExecutor)" Arguments="-precompile -allmodules -nolink $(TargetDebugInfoArg)"/>
+					<Compile Target="$(Target)" Platform="LinuxArm64" Configuration="$(Configuration)" Tag="#UnrealGame LinuxArm64 LinuxNative" AllowParallelExecutor="$(bAllowParallelExecutor)" Arguments="-precompile $(TargetDebugInfoArg)" Clean="false"/>
 				</ForEach>
 			</ForEach>
 
@@ -362,8 +367,8 @@
 		<Node Name="Compile UnrealGame Android" Requires="Compile UnrealHeaderTool $(HostPlatform)" Produces="#UnrealGame Android;#UnrealGame Android Unstripped;#UnrealGame Android Stripped">
 			<ForEach Name="Target" Values="UnrealGame;$(OptionalClientTarget)">
 				<ForEach Name="Configuration" Values="$(GameConfigurations)">
-					<Compile Target="$(Target)" Platform="Android" Configuration="$(Configuration)" Tag="#UnrealGame Android" Arguments="-precompile -allmodules -nolink -architectures=$(AndroidArchitectures) $(TargetDebugInfoArg)"/>
-					<Compile Target="$(Target)" Platform="Android" Configuration="$(Configuration)" Tag="#UnrealGame Android" Arguments="-precompile -architectures=$(AndroidArchitectures) $(TargetDebugInfoArg)" Clean="false"/>
+					<Compile Target="$(Target)" Platform="Android" Configuration="$(Configuration)" Tag="#UnrealGame Android" AllowParallelExecutor="$(bAllowParallelExecutor)" Arguments="-precompile -allmodules -nolink -architectures=$(AndroidArchitectures) $(TargetDebugInfoArg)"/>
+					<Compile Target="$(Target)" Platform="Android" Configuration="$(Configuration)" Tag="#UnrealGame Android" AllowParallelExecutor="$(bAllowParallelExecutor)" Arguments="-precompile -architectures=$(AndroidArchitectures) $(TargetDebugInfoArg)" Clean="false"/>
 				</ForEach>
 			</ForEach>
 
@@ -394,8 +399,8 @@
 			</ForEach>
 			<ForEach Name="Target" Values="UnrealGame;$(OptionalClientTarget);$(OptionalServerTarget)">
 				<ForEach Name="Configuration" Values="$(GameConfigurations)">
-					<Compile Target="$(Target)" Platform="Mac" Configuration="$(Configuration)" Tag="#UnrealGame Mac" Arguments="-precompile -allmodules -nolink $(TargetDebugInfoArg)"/>
-					<Compile Target="$(Target)" Platform="Mac" Configuration="$(Configuration)" Tag="#UnrealGame Mac" Arguments="-precompile $(TargetDebugInfoArg)" Clean="false"/>
+					<Compile Target="$(Target)" Platform="Mac" Configuration="$(Configuration)" Tag="#UnrealGame Mac" AllowParallelExecutor="$(bAllowParallelExecutor)" Arguments="-precompile -allmodules -nolink $(TargetDebugInfoArg)"/>
+					<Compile Target="$(Target)" Platform="Mac" Configuration="$(Configuration)" Tag="#UnrealGame Mac" AllowParallelExecutor="$(bAllowParallelExecutor)" Arguments="-precompile $(TargetDebugInfoArg)" Clean="false"/>
 				</ForEach>
 			</ForEach>
 
@@ -415,8 +420,8 @@
 		<Node Name="Compile UnrealGame IOS" Requires="Compile UnrealHeaderTool Mac" Produces="#UnrealGame IOS;#UnrealGame IOS Unstripped;#UnrealGame IOS Stripped">
 			<ForEach Name="Target" Values="UnrealGame;$(OptionalClientTarget)">
 				<ForEach Name="Configuration" Values="$(GameConfigurations)">
-					<Compile Target="$(Target)" Platform="IOS" Configuration="$(Configuration)" Tag="#UnrealGame IOS" Arguments="-precompile -allmodules -nolink $(TargetDebugInfoArg)"/>
-					<Compile Target="$(Target)" Platform="IOS" Configuration="$(Configuration)" Tag="#UnrealGame IOS" Arguments="-precompile -createstub $(TargetDebugInfoArg)" Clean="false"/>
+					<Compile Target="$(Target)" Platform="IOS" Configuration="$(Configuration)" Tag="#UnrealGame IOS" AllowParallelExecutor="$(bAllowParallelExecutor)" Arguments="-precompile -allmodules -nolink $(TargetDebugInfoArg)"/>
+					<Compile Target="$(Target)" Platform="IOS" Configuration="$(Configuration)" Tag="#UnrealGame IOS" AllowParallelExecutor="$(bAllowParallelExecutor)" Arguments="-precompile -createstub $(TargetDebugInfoArg)" Clean="false"/>
 				</ForEach>
 			</ForEach>
 
@@ -430,8 +435,8 @@
 		<Node Name="Compile UnrealGame TVOS" Requires="Compile UnrealHeaderTool Mac" Produces="#UnrealGame TVOS;#UnrealGame TVOS Unstripped;#UnrealGame TVOS Stripped">
 			<ForEach Name="Target" Values="UnrealGame;$(OptionalClientTarget)">
 				<ForEach Name="Configuration" Values="$(GameConfigurations)">
-					<Compile Target="$(Target)" Platform="TVOS" Configuration="$(Configuration)" Tag="#UnrealGame TVOS" Arguments="-precompile -allmodules -nolink $(TargetDebugInfoArg)"/>
-					<Compile Target="$(Target)" Platform="TVOS" Configuration="$(Configuration)" Tag="#UnrealGame TVOS" Arguments="-precompile -createstub $(TargetDebugInfoArg)" Clean="false"/>
+					<Compile Target="$(Target)" Platform="TVOS" Configuration="$(Configuration)" Tag="#UnrealGame TVOS" AllowParallelExecutor="$(bAllowParallelExecutor)" Arguments="-precompile -allmodules -nolink $(TargetDebugInfoArg)"/>
+					<Compile Target="$(Target)" Platform="TVOS" Configuration="$(Configuration)" Tag="#UnrealGame TVOS" AllowParallelExecutor="$(bAllowParallelExecutor)" Arguments="-precompile -createstub $(TargetDebugInfoArg)" Clean="false"/>
 				</ForEach>
 			</ForEach>
 
