Ensure Build Tool CS Binaries are properly filtered when generating an installed build

Setup a new filter property to filter when tagging before copying to the final install path

#rnx

diff --git a/Engine/Build/InstalledEngineBuild.xml b/Engine/Build/InstalledEngineBuild.xml
--- a/Engine/Build/InstalledEngineBuild.xml
+++ b/Engine/Build/InstalledEngineBuild.xml
@@ -1011,8 +1011,9 @@
 
 			<!-- Tag files that need to be copied directly -->
 			<Tag Files="#Installed Win64" Except="$(CopyInstalledExceptions);#Strip Tools for Installed Win64;#Sign Tools for Installed Win64" With="#Copy for Installed Win64"/>
+			<Tag Files="#Build Tools CS Binaries" Except="$(CopyWin64CsToolsExceptions)" With="#Installed Build Tools CS Binaries"/>
 			<Copy Files="#Copy for Installed Win64" From="$(RootDir)" To="$(LocalInstalledDir)" Tag="#Installed Build Win64 Files"/>
-			<Copy Files="#Build Tools CS Binaries" From="$(CsToolsDir)" To="$(LocalInstalledDir)" Tag="#Installed Build Win64 Files"/>
+			<Copy Files="#Installed Build Tools CS Binaries" From="$(CsToolsDir)" To="$(LocalInstalledDir)" Tag="#Installed Build Win64 Files"/>
 			
 			<!-- Tag saved output that needs to be copied directly -->
 			<Tag BaseDir="$(SavedOutput)" Files="#Saved Output" Except="$(CopyInstalledExceptions)" With="#Copy Saved Output"/>
@@ -1139,8 +1140,9 @@
 
 			<!-- Tag files that need to be copied directly -->
 			<Tag Files="#Installed Mac" Except="$(CopyInstalledExceptions);#Strip Tools for Installed Mac;#Sign Tools for Installed Mac" With="#Copy for Installed Mac"/>
+			<Tag Files="#Build Tools CS Binaries Mac" Except="$(CopyMacCsToolsExceptions)" With="#Installed Build Tools CS Binaries Mac"/>
 			<Copy Files="#Copy for Installed Mac" From="$(RootDir)" To="$(LocalInstalledDirMac)" Tag="#Installed Build Mac Files"/>
-			<Copy Files="#Build Tools CS Binaries Mac" From="$(CsToolsDir)" To="$(LocalInstalledDirMac)" Tag="#Installed Build Mac Files"/>
+			<Copy Files="#Installed Build Tools CS Binaries Mac" From="$(CsToolsDir)" To="$(LocalInstalledDirMac)" Tag="#Installed Build Mac Files"/>
 			
 			<!-- Tag saved output that needs to be copied directly -->
 			<Tag BaseDir="$(SavedOutput)" Files="#Saved Output" Except="$(CopyInstalledExceptions)" With="#Copy Saved Output"/>
@@ -1255,8 +1257,9 @@
 
 			<!-- Tag saved output that needs to be copied directly -->
 			<Tag BaseDir="$(SavedOutput)" Files="#Saved Output" Except="$(CopyInstalledExceptions)" With="#Copy Saved Output"/>
+			<Tag Files="#Build Tools CS Binaries LinuxNative" Except="$(CopyLinuxCsToolsExceptions)" With="#Installed Build Tools CS Binaries LinuxNative"/>
 			<Copy Files="#Copy Saved Output" From="$(SavedOutput)" To="$(LocalInstalledDirLinux)" Tag="#Installed Build Linux Files"/>
-			<Copy Files="#Build Tools CS Binaries LinuxNative" From="$(CsToolsDir)" To="$(LocalInstalledDirLinux)" Tag="#Installed Build Linux Files"/>
+			<Copy Files="#Installed Build Tools CS Binaries LinuxNative" From="$(CsToolsDir)" To="$(LocalInstalledDirLinux)" Tag="#Installed Build Linux Files"/>
 			
 			<!-- Strip required files into their final location -->
 			<Strip Files="#Strip Tools for Installed Linux" BaseDir="$(RootDir)" OutputDir="$(LocalInstalledDirLinux)" Platform="Linux" Tag="#Installed Build Linux Files"/>
diff --git a/Engine/Build/InstalledEngineFilters.xml b/Engine/Build/InstalledEngineFilters.xml
--- a/Engine/Build/InstalledEngineFilters.xml
+++ b/Engine/Build/InstalledEngineFilters.xml
@@ -430,7 +430,11 @@
 		Engine/Source/Programs/Shared/*/bin/...
 		Engine/Source/Programs/Shared/*/obj/...
 	</Property>
-	
+	<Property Name="CopyWin64CsToolsExceptions">
+		Engine/Saved/CsTools/Engine/Source/Programs/Shared/*/bin/...
+		Engine/Saved/CsTools/Engine/Source/Programs/Shared/*/obj/...
+	</Property>
+
 	<!-- Mac -->
 	<Property Name="CopyMacFilter">
 		Engine/Binaries/Mac/BootstrapPackagedGame.app/...
@@ -449,6 +453,10 @@
 		Engine/Source/Programs/Shared/*/bin/...
 		Engine/Source/Programs/Shared/*/obj/...
 	</Property>
+	<Property Name="CopyMacCsToolsExceptions">
+		Engine/Saved/CsTools/Engine/Source/Programs/Shared/*/bin/...
+		Engine/Saved/CsTools/Engine/Source/Programs/Shared/*/obj/...
+	</Property>
 
 	<!-- Android -->
 	<Property Name="CopyAndroidFilter">
@@ -567,6 +575,11 @@
 		Engine/Source/Programs/Shared/*/obj/...
 	</Property>
 
+	<Property Name="CopyLinuxCsToolsExceptions">
+		Engine/Saved/CsTools/Engine/Source/Programs/Shared/*/bin/...
+		Engine/Saved/CsTools/Engine/Source/Programs/Shared/*/obj/...
+	</Property>
+
 	<Property Name="CopyLinuxFilter" Value="$(CopyLinuxFilter);Engine/Extras/ThirdPartyNotUE/putty/..."/>	
 
 	<!-- Win64 specific -->
