Allow enabling of DYSM files for Mac with an Installed Engine Build

* Change -NoDSYM arg to -EnableDSYM to allow enabling of DYSM files for Mac
* Plumb that setting through InstalledEngineBuild.xml to allow generating them for an engine build

#rnx

diff --git a/Engine/Build/InstalledEngineBuild.xml b/Engine/Build/InstalledEngineBuild.xml
--- a/Engine/Build/InstalledEngineBuild.xml
+++ b/Engine/Build/InstalledEngineBuild.xml
@@ -99,6 +99,8 @@
 	<!-- Compile flags for Windows Game binaries -->
 	<Property Name="TargetDebugInfoArg" Value="-nodebuginfo"/>
 	<Property Name="TargetDebugInfoArg" Value="" If="$(WithFullDebugInfo)"/>
+	<Property Name="TargetDebugInfoArgMac" Value="$(TargetDebugInfoArg)" />
+	<Property Name="TargetDebugInfoArgMac" Value="-EnableDSYM" If="!$(WithFullDebugInfo) And $(WithMac)" />
 
 	<!-- Names of the client and server targets to build -->
 	<Property Name="OptionalClientTarget" Value=""/>
@@ -198,8 +200,8 @@
 		<Node Name="Compile UnrealEditor Mac" Requires="Compile UnrealHeaderTool Mac" Produces="#UnrealEditor Mac;#UnrealEditor Mac Unstripped;#UnrealEditor Mac Stripped;#UnrealEditor Mac Unsigned;#UnrealEditor Mac Signed;#UnrealEditor Mac WithExceptions">
 			<Property Name="AllModulesArg" Value="-allmodules"/>
 			<Property Name="AllModulesArg" Value="" If="$(WithMacArm64)"/>
-			<Compile Target="UnrealEditor" Platform="Mac" Configuration="DebugGame" Tag="#UnrealEditor Mac" Arguments="-precompile $(AllModulesArg) $(CrashReporterCompileArgs) $(UseArm64ArchitectureArg)"/>
-			<Compile Target="UnrealEditor" Platform="Mac" Configuration="Development" Tag="#UnrealEditor Mac" Arguments="-precompile $(AllModulesArg) $(CrashReporterCompileArgs) $(UseArm64ArchitectureArg)"/>
+			<Compile Target="UnrealEditor" Platform="Mac" Configuration="DebugGame" Tag="#UnrealEditor Mac" Arguments="-precompile $(AllModulesArg) $(CrashReporterCompileArgs) $(UseArm64ArchitectureArg) $(TargetDebugInfoArgMac)"/>
+			<Compile Target="UnrealEditor" Platform="Mac" Configuration="Development" Tag="#UnrealEditor Mac" Arguments="-precompile $(AllModulesArg) $(CrashReporterCompileArgs) $(UseArm64ArchitectureArg) $(TargetDebugInfoArgMac)"/>
 
 			<Property Name="UnrealEditorExceptions" Value=""/>
 			<Property Name="UnrealEditorExceptions" Value="$(PluginsExceptions)" If="$(CompileDatasmithPlugins)"/>
@@ -386,8 +388,8 @@
 		<ForEach Name="Target" Values="$(OptionalClientTarget);$(OptionalServerTarget)">
 			<Node Name="Compile $(Target) Mac" Requires="Compile UnrealHeaderTool Mac" Produces="#$(Target) Mac">
 				<ForEach Name="Configuration" Values="$(GameConfigurations)">
-					<Compile Target="$(Target)" Platform="Mac" Configuration="$(Configuration)" Tag="#$(Target) Mac" Arguments="-precompile -allmodules -nolink $(TargetDebugInfoArg)"/>
-					<Compile Target="$(Target)" Platform="Mac" Configuration="$(Configuration)" Tag="#$(Target) Mac" Arguments="-precompile $(TargetDebugInfoArg)" Clean="false"/>
+					<Compile Target="$(Target)" Platform="Mac" Configuration="$(Configuration)" Tag="#$(Target) Mac" Arguments="-precompile -allmodules -nolink $(TargetDebugInfoArgMac)"/>
+					<Compile Target="$(Target)" Platform="Mac" Configuration="$(Configuration)" Tag="#$(Target) Mac" Arguments="-precompile $(TargetDebugInfoArgMac)" Clean="false"/>
 					<Property Name="OptionalTargetTags" Value="$(OptionalTargetTags);#$(Target) Mac"/>
 				</ForEach>
 			</Node>
@@ -399,8 +401,8 @@
 			</ForEach>
 			<ForEach Name="Target" Values="UnrealGame;$(OptionalClientTarget);$(OptionalServerTarget)">
 				<ForEach Name="Configuration" Values="$(GameConfigurations)">
-					<Compile Target="$(Target)" Platform="Mac" Configuration="$(Configuration)" Tag="#UnrealGame Mac" AllowParallelExecutor="$(bAllowParallelExecutor)" Arguments="-precompile -allmodules -nolink $(TargetDebugInfoArg)"/>
-					<Compile Target="$(Target)" Platform="Mac" Configuration="$(Configuration)" Tag="#UnrealGame Mac" AllowParallelExecutor="$(bAllowParallelExecutor)" Arguments="-precompile $(TargetDebugInfoArg)" Clean="false"/>
+					<Compile Target="$(Target)" Platform="Mac" Configuration="$(Configuration)" Tag="#UnrealGame Mac" AllowParallelExecutor="$(bAllowParallelExecutor)" Arguments="-precompile -allmodules -nolink $(TargetDebugInfoArgMac)"/>
+					<Compile Target="$(Target)" Platform="Mac" Configuration="$(Configuration)" Tag="#UnrealGame Mac" AllowParallelExecutor="$(bAllowParallelExecutor)" Arguments="-precompile $(TargetDebugInfoArgMac)" Clean="false"/>
 				</ForEach>
 			</ForEach>
 
@@ -420,8 +422,8 @@
 		<Node Name="Compile UnrealGame IOS" Requires="Compile UnrealHeaderTool Mac" Produces="#UnrealGame IOS;#UnrealGame IOS Unstripped;#UnrealGame IOS Stripped">
 			<ForEach Name="Target" Values="UnrealGame;$(OptionalClientTarget)">
 				<ForEach Name="Configuration" Values="$(GameConfigurations)">
-					<Compile Target="$(Target)" Platform="IOS" Configuration="$(Configuration)" Tag="#UnrealGame IOS" AllowParallelExecutor="$(bAllowParallelExecutor)" Arguments="-precompile -allmodules -nolink $(TargetDebugInfoArg)"/>
-					<Compile Target="$(Target)" Platform="IOS" Configuration="$(Configuration)" Tag="#UnrealGame IOS" AllowParallelExecutor="$(bAllowParallelExecutor)" Arguments="-precompile -createstub $(TargetDebugInfoArg)" Clean="false"/>
+					<Compile Target="$(Target)" Platform="IOS" Configuration="$(Configuration)" Tag="#UnrealGame IOS" AllowParallelExecutor="$(bAllowParallelExecutor)" Arguments="-precompile -allmodules -nolink $(TargetDebugInfoArgMac)"/>
+					<Compile Target="$(Target)" Platform="IOS" Configuration="$(Configuration)" Tag="#UnrealGame IOS" AllowParallelExecutor="$(bAllowParallelExecutor)" Arguments="-precompile -createstub $(TargetDebugInfoArgMac)" Clean="false"/>
 				</ForEach>
 			</ForEach>
 
@@ -435,8 +437,8 @@
 		<Node Name="Compile UnrealGame TVOS" Requires="Compile UnrealHeaderTool Mac" Produces="#UnrealGame TVOS;#UnrealGame TVOS Unstripped;#UnrealGame TVOS Stripped">
 			<ForEach Name="Target" Values="UnrealGame;$(OptionalClientTarget)">
 				<ForEach Name="Configuration" Values="$(GameConfigurations)">
-					<Compile Target="$(Target)" Platform="TVOS" Configuration="$(Configuration)" Tag="#UnrealGame TVOS" AllowParallelExecutor="$(bAllowParallelExecutor)" Arguments="-precompile -allmodules -nolink $(TargetDebugInfoArg)"/>
-					<Compile Target="$(Target)" Platform="TVOS" Configuration="$(Configuration)" Tag="#UnrealGame TVOS" AllowParallelExecutor="$(bAllowParallelExecutor)" Arguments="-precompile -createstub $(TargetDebugInfoArg)" Clean="false"/>
+					<Compile Target="$(Target)" Platform="TVOS" Configuration="$(Configuration)" Tag="#UnrealGame TVOS" AllowParallelExecutor="$(bAllowParallelExecutor)" Arguments="-precompile -allmodules -nolink $(TargetDebugInfoArgMac)"/>
+					<Compile Target="$(Target)" Platform="TVOS" Configuration="$(Configuration)" Tag="#UnrealGame TVOS" AllowParallelExecutor="$(bAllowParallelExecutor)" Arguments="-precompile -createstub $(TargetDebugInfoArgMac)" Clean="false"/>
 				</ForEach>
 			</ForEach>
 
diff --git a/Engine/Source/Programs/UnrealBuildTool/Platform/Mac/UEBuildMac.cs b/Engine/Source/Programs/UnrealBuildTool/Platform/Mac/UEBuildMac.cs
--- a/Engine/Source/Programs/UnrealBuildTool/Platform/Mac/UEBuildMac.cs
+++ b/Engine/Source/Programs/UnrealBuildTool/Platform/Mac/UEBuildMac.cs
@@ -1,4 +1,4 @@
-﻿// Copyright Epic Games, Inc. All Rights Reserved.
+// Copyright Epic Games, Inc. All Rights Reserved.
 
 
 using System;
@@ -39,7 +39,7 @@ namespace UnrealBuildTool
 		/// <summary>
 		/// Enables the generation of .dsym files. This can be disabled to enable faster iteration times during development.
 		/// </summary>
-		[CommandLine("-NoDSYM", Value = "false")]
+		[CommandLine("-EnableDSYM", Value = "true")]
 		[XmlConfigFile(Category = "BuildConfiguration", Name = "bUseDSYMFiles")]
 		public bool bUseDSYMFiles = false;
 
