Removing force enabling AllowParallelExecutor

Don't force enable if IsBuildMachine is eanbled.

#rnx

diff --git a/Engine/Build/InstalledEngineBuild.xml b/Engine/Build/InstalledEngineBuild.xml
--- a/Engine/Build/InstalledEngineBuild.xml
+++ b/Engine/Build/InstalledEngineBuild.xml
@@ -72,7 +72,6 @@
 	<!-- Control ParallelExecutor -->
 	<Option	Name="AllowParallelExecutor" Restrict="true|false" DefaultValue="false" Description="Allow the use of UAT's ParallelExecutor when compiling"/>
 	<Property Name="bAllowParallelExecutor" Value="$(AllowParallelExecutor)" />
-	<Property Name="bAllowParallelExecutor" Value="true" If="$(IsBuildMachine)" />
 
 	<!-- The local output directory -->
 	<Property Name="LocalInstalledDir" Value="$(BuiltDirectory)/Windows"/>
