Adjusted the naming of Editor Perf tests
Split Zen into 2 separate tests with toggles for zenloader and async loading
Adjusted location of traces and csv storage to be representative of the job link

diff --git a/Engine/Build/EditorPerf.xml b/Engine/Build/EditorPerf.xml
--- a/Engine/Build/EditorPerf.xml
+++ b/Engine/Build/EditorPerf.xml
@@ -68,14 +68,16 @@
 	<Option Name="Iterations" DefaultValue="1" Description="The list of iteration numbers to run. For example '1;2;3' for three iterations"/>
 
 	<!-- Arguments for extra Zen Loader tests -->
-	<Option Name="EnableZenLoaderTests" DefaultValue="false" Description="Enable additional test per phase that uses the Zen Loader."/>
-	<Option Name="EnableZenAsyncThread" DefaultValue="true" Description="Enable additional test per phase that uses the Zen Loader."/>
-	<Option Name="EnableZenWithoutEditorDomain" DefaultValue="true" Description="Enable additional test per phase that uses the Zen Loader."/>
+	<Option Name="EnableZenLoaderTest" DefaultValue="true" Description="Enable additional test per phase that uses the Zen Loader."/>
+	<Option Name="EnableZenAsyncThreadTest" DefaultValue="true" Description="Enable additional test per phase that uses the Zen Loader."/>
+	<Option Name="EnableZenTestsWithoutEditorDomain" DefaultValue="true" Description="Enable additional test per phase that uses the Zen Loader."/>
 
 	<!-- Zen Loader tests specifc arguments -->
-	<Property Name="ZenLoaderArgs" Value="-zenloader"/>
-	<Property Name="ZenLoaderArgs" Value="$(ZenLoaderArgs),-asyncloadingthread" If="$(EnableZenAsyncThread)" />
-	<Property Name="ZenLoaderArgs" Value="$(ZenLoaderArgs),-noeditordomain" If="$(EnableZenWithoutEditorDomain)" />
+	<Property Name="ZenLoaderTestArgs" Value="-zenloader"/>
+	<Property Name="ZenLoaderAsyncTestArgs" Value="$(ZenLoaderTestArgs),-asyncloadingthread" If="$(EnableZenAsyncThreadTest)" />
+	
+	<Property Name="ZenLoaderTestArgs" Value="$(ZenLoaderTestArgs),-noeditordomain" If="$(EnableZenTestsWithoutEditorDomain)" />
+	<Property Name="ZenLoaderAsyncTestArgs" Value="$(ZenLoaderAsyncTestArgs),-noeditordomain" If="$(EnableZenTestsWithoutEditorDomain)" />
 
 	<!-- Gauntlet arguments for all runs -->
 	<Property Name="GauntletArgs" Value="-build=editor -UseEditor -verbose -tempdir=&quot;$(RootDir)/Tests&quot;"/>
@@ -101,21 +103,23 @@
 	<Property Name="TelemetryConfig" Value="UETelemetry" If="$(IsBuildMachine) And '$(TelemetryConfig)' == ''"/>
 
 	<!-- The main test running macro logic that is expanded with platform and dependency node names -->
-	<Macro Name="RunTest" Arguments="OutputName;Platform;Project;NodeName;Phase;Iteration;Args;WithZen">
+	<Macro Name="RunTest" Arguments="TestNameSuffix;Platform;Project;NodeName;Phase;Iteration;Args;ZenArgs">
 
-		<Property Name="TestName" Value="EditorPerf$(Phase)"/>
-		<Property Name="TestName" Value="$(TestName)Zen" If="$(WithZen)"/>
-		<Property Name="IterationName" Value="$(Project)$(Platform)$(TestName)$(Iteration)"/>
+		<Property Name="TestName" Value="EditorPerf$(Phase)$(TestNameSuffix)"/>
+		<Property Name="IterationSuffix" Value=""/>
+		<Property Name="IterationSuffix" Value="$(Iteration)" If="'$(Iteration)' != '1'"/>
+
+		<Property Name="TestOutputDir" Value="$(AutomationOutputDir)/$(Project)$(Platform)EditorPerf"/>
+
+		<Property Name="IterationName" Value="$(Project)$(Platform)$(TestName)$(IterationSuffix)"/>
 		<Property Name="TraceFilePath" Value="$(TraceDir)/$(IterationName).utrace"/>
 		<Property Name="TelemetryFileName" Value="$(IterationName)Telemetry.csv"/>
 		<Property Name="TelemetryFilePath" Value="$(TraceDir)/$(TelemetryFileName)"/>
-		<Property Name="LogOutputParams" Value="-logdir=&quot;$(AutomationOutputDir)/$(OutputName)&quot;"/>
+		<Property Name="LogOutputParams" Value="-logdir=&quot;$(TestOutputDir)&quot;"/>
 		<Property Name="ReportParams" Value="-WriteTestResultsForHorde -HordeTestDataKey=&quot;$(NodeName)$(PreflightSuffix)&quot;"/>
-		<Property Name="ArtifactName" Value="$(Phase)$(Iteration)"/>
-		<Property Name="ArtifactName" Value="$(Phase)$(Iteration)" If="$(WithZen)"/>
+		<Property Name="ArtifactName" Value="$(Phase)$(IterationSuffix)$(TestNameSuffix)"/>
 		<Property Name="AllGauntletArgs" Value="-Project=$(Project) -Platform=$(Platform) $(GauntletArgs) $(ReportParams) $(LogOutputParams) -artifactname=&quot;$(ArtifactName)&quot;"/>
-		<Property Name="RunArgs" Value="$(Args)"/>
-		<Property Name="RunArgs" Value="$(Args) -editorargs=&quot;$(ZenLoaderArgs)&quot;" If="$(WithZen)"/>
+		<Property Name="RunArgs" Value="$(Args) -editorargs=&quot;$(ZenArgs)&quot;"/>
 
 		<!-- Do DDC cleanup outside the editor process, handles Cold cleanup without changing stats -->
 		<Commandlet Project="$(Project)" Name="DDCCleanup" Arguments="$(Args) $(CommandletArgs)"/>
@@ -127,7 +131,7 @@
 		<Commandlet Project="$(Project)" Name="SummarizeTrace" Arguments="-inputfile=&quot;$(TraceFilePath)&quot; -testname=&quot;$(TestName)&quot; $(SummarizeArgs) $(CommandletArgs)"/>
 
 		<!-- Save all files named for this run in the trace directory -->
-		<Copy From="$(TraceDir)/$(IterationName)*.*" To="$(AutomationOutputDir)"/>
+		<Copy From="$(TraceDir)/$(IterationName)*.*" To="$(TestOutputDir)"/>
 
 		<Do If="'$(TelemetryConfig)' != ''">
 			<!-- Publish telemetry to the db -->
@@ -190,18 +194,27 @@
 					<ForEach Name="Iteration" Values="$(Iterations)">
 						<!-- Only ever do one iteration of incremental, because it will warm up the workspace -->
 						<Do If="'$(Phase)' != 'Incr' or '$(Iteration)' == '1'">
-							<Property Name="RunNodeName" Value="Run $(OutputName) $(Phase)$(Iteration)"/>
+							<Property Name="IterationSuffix" Value=""/>
+							<Property Name="IterationSuffix" Value="$(Iteration)" If="'$(Iteration)' != '1'"/>
+							<Property Name="RunNodeName" Value="Run $(OutputName) $(Phase)$(IterationSuffix)"/>
 							<Node Name="$(RunNodeName)" Requires="$(RunNodeRequires);$(RunNodeList)">
-								<Expand Name="RunTest" OutputName="$(OutputName)" Platform="$(Platform)" Project="$(Project)" NodeName="$(RunNodeName)" Phase="$(Phase)" Iteration="$(Iteration)" Args="$($(Phase)Args)" WithZen="false"/>
+								<Expand Name="RunTest" TestNameSuffix="" Platform="$(Platform)" Project="$(Project)" NodeName="$(RunNodeName)" Phase="$(Phase)" Iteration="$(Iteration)" Args="$($(Phase)Args)" ZenArgs=""/>
 							</Node>
 							<Property Name="RunNodeList" Value="$(RunNodeList);$(RunNodeName)"/>
 						</Do>
 					</ForEach>
+					<!-- Zen Loader TestRun if enabled -->
 					<Property Name="RunNodeZenName" Value="Run $(OutputName) $(Phase) Zen"/>
-					<Node Name="$(RunNodeZenName)" Requires="$(RunNodeRequires);$(RunNodeList)" If="$(EnableZenLoaderTests)">
-						<Expand Name="RunTest" OutputName="$(OutputName)Zen" Platform="$(Platform)" Project="$(Project)" NodeName="$(RunNodeZenName)" Phase="$(Phase)" Iteration="1" Args="$($(Phase)Args)" WithZen="true"/>
+					<Node Name="$(RunNodeZenName)" Requires="$(RunNodeRequires);$(RunNodeList)" If="$(EnableZenLoaderTest)">
+						<Expand Name="RunTest" TestNameSuffix="Zen" Platform="$(Platform)" Project="$(Project)" NodeName="$(RunNodeZenName)" Phase="$(Phase)" Iteration="1" Args="$($(Phase)Args)" ZenArgs="$(ZenLoaderTestArgs)"/>
+					</Node>
+					<Property Name="RunNodeList" Value="$(RunNodeList);$(RunNodeZenName)" If="$(EnableZenLoaderTest)"/>
+					<!-- Zen Loader Async TestRun if enabled -->
+					<Property Name="RunNodeZenAsyncName" Value="Run $(OutputName) $(Phase) Zen Async"/>
+					<Node Name="$(RunNodeZenAsyncName)" Requires="$(RunNodeRequires);$(RunNodeList)" If="$(EnableZenAsyncThreadTest)">
+						<Expand Name="RunTest" TestNameSuffix="ZenAsync" Platform="$(Platform)" Project="$(Project)" NodeName="$(RunNodeZenAsyncName)" Phase="$(Phase)" Iteration="1" Args="$($(Phase)Args)" ZenArgs="$(ZenLoaderAsyncTestArgs)"/>
 					</Node>
-					<Property Name="RunNodeList" Value="$(RunNodeList);$(RunNodeZenName)" If="$(EnableZenLoaderTests)"/>
+					<Property Name="RunNodeList" Value="$(RunNodeList);$(RunNodeZenAsyncName)" If="$(EnableZenAsyncThreadTest)"/>
 				</ForEach>
 			</Agent>
 			<Label Category="$(RunCategory)" Name="$(Project) Editor Run" Requires="$(RunNodeRequires);$(RunNodeList)" UgsBadge="$(Project) Editor Run"/>
