Disable dsym generation due to a symbol failing to be dumped by breakpad

#rnx

diff --git a/Engine/Build/InstalledEngineBuild.xml b/Engine/Build/InstalledEngineBuild.xml
--- a/Engine/Build/InstalledEngineBuild.xml
+++ b/Engine/Build/InstalledEngineBuild.xml
@@ -99,7 +99,7 @@
 	<Property Name="TargetDebugInfoArg" Value="-nodebuginfo"/>
 	<Property Name="TargetDebugInfoArg" Value="" If="$(WithFullDebugInfo)"/>
 	<Property Name="TargetDebugInfoArgMac" Value="$(TargetDebugInfoArg)" />
-	<Property Name="TargetDebugInfoArgMac" Value="-EnableDSYM" If="!$(WithFullDebugInfo) And $(WithMac)" />
+	<!-- <Property Name="TargetDebugInfoArgMac" Value="-EnableDSYM" If="!$(WithFullDebugInfo) And $(WithMac)" /> -->
 
 	<!-- Names of the client and server targets to build -->
 	<Property Name="OptionalClientTarget" Value=""/>
