Fixing EpicGames.Perforce.Native not being included in an installed linux build

#rnx

diff --git a/Engine/Build/InstalledEngineBuild.xml b/Engine/Build/InstalledEngineBuild.xml
--- a/Engine/Build/InstalledEngineBuild.xml
+++ b/Engine/Build/InstalledEngineBuild.xml
@@ -665,7 +665,7 @@
 			<Tag Files="$(CsToolsDir)/Engine/Binaries/DotNET/UnrealBuildTool/..." With="#Build Tools CS Binaries LinuxNative" />
 
 			<!-- Tag the native P4 libraries as they are not a C# project so not a reference project -->
-			<Tag Files="$(CsToolsDir)/Engine/Binaries/DotNET/EpicGames.Perforce.Native/..." With="#Build Tools CS LinuxNative" />
+			<Tag Files="$(CsToolsDir)/Engine/Binaries/DotNET/EpicGames.Perforce.Native/..." With="#Build Tools CS Binaries LinuxNative" />
 
 			<!-- Tag AutomationTool Script module build records, so that prebuilt modules may be discovered in the absence of source code -->
 			<Tag Files="$(CsToolsDir)/Engine/Intermediate/ScriptModules/..." With="#Build Tools CS Binaries LinuxNative" />
