Correctly re-enabling dsym generation for mac

Force the editor & tools to have symbols, game targets are optional based on WithFullDebugInfo

#rnx

diff --git a/Engine/Build/InstalledEngineBuild.xml b/Engine/Build/InstalledEngineBuild.xml
--- a/Engine/Build/InstalledEngineBuild.xml
+++ b/Engine/Build/InstalledEngineBuild.xml
@@ -99,7 +99,7 @@
 	<Property Name="TargetDebugInfoArg" Value="-nodebuginfo"/>
 	<Property Name="TargetDebugInfoArg" Value="" If="$(WithFullDebugInfo)"/>
 	<Property Name="TargetDebugInfoArgMac" Value="$(TargetDebugInfoArg)" />
-	<!-- <Property Name="TargetDebugInfoArgMac" Value="-EnableDSYM" If="!$(WithFullDebugInfo) And $(WithMac)" /> -->
+	<Property Name="TargetDebugInfoArgMac" Value="-EnableDSYM" If="$(WithFullDebugInfo)" />
 
 	<!-- Names of the client and server targets to build -->
 	<Property Name="OptionalClientTarget" Value=""/>
@@ -199,8 +199,8 @@
 		<Node Name="Compile UnrealEditor Mac" Requires="Compile UnrealHeaderTool Mac" Produces="#UnrealEditor Mac;#UnrealEditor Mac Unstripped;#UnrealEditor Mac Stripped;#UnrealEditor Mac Unsigned;#UnrealEditor Mac Signed;#UnrealEditor Mac WithExceptions">
 			<Property Name="AllModulesArg" Value="-allmodules"/>
 			<Property Name="AllModulesArg" Value="" If="$(WithMacArm64)"/>
-			<Compile Target="UnrealEditor" Platform="Mac" Configuration="DebugGame" Tag="#UnrealEditor Mac" Arguments="-precompile $(AllModulesArg) $(CrashReporterCompileArgs) $(UseArm64ArchitectureArg) $(TargetDebugInfoArgMac)"/>
-			<Compile Target="UnrealEditor" Platform="Mac" Configuration="Development" Tag="#UnrealEditor Mac" Arguments="-precompile $(AllModulesArg) $(CrashReporterCompileArgs) $(UseArm64ArchitectureArg) $(TargetDebugInfoArgMac)"/>
+			<Compile Target="UnrealEditor" Platform="Mac" Configuration="DebugGame" Tag="#UnrealEditor Mac" Arguments="-precompile -EnableDSYM $(AllModulesArg) $(CrashReporterCompileArgs) $(UseArm64ArchitectureArg)"/>
+			<Compile Target="UnrealEditor" Platform="Mac" Configuration="Development" Tag="#UnrealEditor Mac" Arguments="-precompile -EnableDSYM $(AllModulesArg) $(CrashReporterCompileArgs) $(UseArm64ArchitectureArg)"/>
 
 			<Property Name="UnrealEditorExceptions" Value=""/>
 			<Property Name="UnrealEditorExceptions" Value="$(PluginsExceptions)" If="$(CompileDatasmithPlugins)"/>
@@ -607,22 +607,22 @@
 		</Node>
 	
 		<Node Name="Build Tools Mac" Requires="Compile UnrealHeaderTool Mac">
-			<Compile Target="CrashReportClient" Configuration="Shipping" Platform="Mac" Tag="#Build Tools Mac" Arguments="$(CrashReporterCompileArgs)"/>
-			<Compile Target="CrashReportClientEditor" Configuration="Shipping" Platform="Mac" Tag="#Build Tools Mac" Arguments="$(CrashReporterCompileArgs)"/>
-			<Compile Target="ShaderCompileWorker" Configuration="Development" Platform="Mac" Tag="#Build Tools Mac"/>
-			<Compile Target="EpicWebHelper" Configuration="Development" Platform="Mac" Tag="#Build Tools Mac"/>
-			<Compile Target="UnrealInsights" Configuration="Development" Platform="Mac" Tag="#Build Tools Mac"/>
-			<Compile Target="UnrealFrontend" Configuration="Development" Platform="Mac" Tag="#Build Tools Mac"/>
-			<Compile Target="UnrealLightmass" Configuration="Development" Platform="Mac" Tag="#Build Tools Mac"/>
-			<Compile Target="UnrealPak" Configuration="Development" Platform="Mac" Tag="#Build Tools Mac"/>
-			<Compile Target="UnrealMultiUserServer" Configuration="Development" Platform="Mac" Tag="#Build Tools Mac"/>
-			<Compile Target="UnrealMultiUserSlateServer" Configuration="Development" Platform="Mac" Tag="#Build Tools Mac"/>
-			<Compile Target="UnrealRecoverySvc" Configuration="Development" Platform="Mac" Tag="#Build Tools Mac"/>
-			<Compile Target="BuildPatchTool" Configuration="Shipping" Platform="Mac" Tag="#Build Tools Mac"/>
-			<Compile Target="DsymExporter" Configuration="Development" Platform="Mac" Tag="#Build Tools Mac"/>
-			<Compile Target="UnrealEditorServices" Configuration="Development" Platform="Mac" Tag="#Build Tools Mac"/>
-			<Compile Target="UnrealAtoS" Configuration="Development" Platform="Mac" Tag="#Build Tools Mac"/>
-			<Compile Target="UnrealObjectPtrTool" Configuration="Development" Platform="Mac" Tag="#Build Tools Mac"/>
+			<Compile Target="CrashReportClient" Configuration="Shipping" Platform="Mac" Tag="#Build Tools Mac" Arguments="-EnableDSYM $(CrashReporterCompileArgs)"/>
+			<Compile Target="CrashReportClientEditor" Configuration="Shipping" Platform="Mac" Tag="#Build Tools Mac" Arguments="-EnableDSYM $(CrashReporterCompileArgs)"/>
+			<Compile Target="ShaderCompileWorker" Configuration="Development" Platform="Mac" Tag="#Build Tools Mac" Arguments="-EnableDSYM"/>
+			<Compile Target="EpicWebHelper" Configuration="Development" Platform="Mac" Tag="#Build Tools Mac" Arguments="-EnableDSYM"/>
+			<Compile Target="UnrealInsights" Configuration="Development" Platform="Mac" Tag="#Build Tools Mac" Arguments="-EnableDSYM"/>
+			<Compile Target="UnrealFrontend" Configuration="Development" Platform="Mac" Tag="#Build Tools Mac" Arguments="-EnableDSYM"/>
+			<Compile Target="UnrealLightmass" Configuration="Development" Platform="Mac" Tag="#Build Tools Mac" Arguments="-EnableDSYM"/>
+			<Compile Target="UnrealPak" Configuration="Development" Platform="Mac" Tag="#Build Tools Mac" Arguments="-EnableDSYM"/>
+			<Compile Target="UnrealMultiUserServer" Configuration="Development" Platform="Mac" Tag="#Build Tools Mac" Arguments="-EnableDSYM"/>
+			<Compile Target="UnrealMultiUserSlateServer" Configuration="Development" Platform="Mac" Tag="#Build Tools Mac" Arguments="-EnableDSYM"/>
+			<Compile Target="UnrealRecoverySvc" Configuration="Development" Platform="Mac" Tag="#Build Tools Mac" Arguments="-EnableDSYM"/>
+			<Compile Target="BuildPatchTool" Configuration="Shipping" Platform="Mac" Tag="#Build Tools Mac" Arguments="-EnableDSYM"/>
+			<Compile Target="DsymExporter" Configuration="Development" Platform="Mac" Tag="#Build Tools Mac" Arguments="-EnableDSYM"/>
+			<Compile Target="UnrealEditorServices" Configuration="Development" Platform="Mac" Tag="#Build Tools Mac" Arguments="-EnableDSYM"/>
+			<Compile Target="UnrealAtoS" Configuration="Development" Platform="Mac" Tag="#Build Tools Mac" Arguments="-EnableDSYM"/>
+			<Compile Target="UnrealObjectPtrTool" Configuration="Development" Platform="Mac" Tag="#Build Tools Mac" Arguments="-EnableDSYM"/>
 		</Node>
 	</Agent>
 
