Adding missing WebAPI App Plugin files to installed builds

#rnx

diff --git a/Engine/Build/InstalledEngineFilters.xml b/Engine/Build/InstalledEngineFilters.xml
--- a/Engine/Build/InstalledEngineFilters.xml
+++ b/Engine/Build/InstalledEngineFilters.xml
@@ -90,6 +90,7 @@
 		<!-- Plugins -->
 		Engine/Extras/ThirdPartyNotUE/SwitchboardThirdParty/...
 		Engine/Plugins/VirtualProduction/RemoteControlWebInterface/WebApp/...
+		Engine/Plugins/Experimental/Web/WebAPI/WebAPIGeneratorApp/...
 
 		<!-- Source code -->
 		Engine/Source/UnrealGame.Target.cs
