Fix warning about missing property when running without zen tests

diff --git a/Engine/Build/EditorPerf.xml b/Engine/Build/EditorPerf.xml
--- a/Engine/Build/EditorPerf.xml
+++ b/Engine/Build/EditorPerf.xml
@@ -74,7 +74,7 @@
 
 	<!-- Zen Loader tests specifc arguments -->
 	<Property Name="ZenLoaderTestArgs" Value="-zenloader"/>
-	<Property Name="ZenLoaderAsyncTestArgs" Value="$(ZenLoaderTestArgs),-asyncloadingthread" If="$(EnableZenAsyncThreadTest)" />
+	<Property Name="ZenLoaderAsyncTestArgs" Value="$(ZenLoaderTestArgs),-asyncloadingthread" />
 	
 	<Property Name="ZenLoaderTestArgs" Value="$(ZenLoaderTestArgs),-noeditordomain" If="$(EnableZenTestsWithoutEditorDomain)" />
 	<Property Name="ZenLoaderAsyncTestArgs" Value="$(ZenLoaderAsyncTestArgs),-noeditordomain" If="$(EnableZenTestsWithoutEditorDomain)" />
