Fix not building EpicWebHelper for installed builds for Linux

diff --git a/Engine/Build/InstalledEngineBuild.xml b/Engine/Build/InstalledEngineBuild.xml
--- a/Engine/Build/InstalledEngineBuild.xml
+++ b/Engine/Build/InstalledEngineBuild.xml
@@ -547,6 +547,7 @@
 			<Compile Target="CrashReportClient" Platform="Linux" Configuration="Shipping" Tag="#Build Tools Linux" Arguments="$(CrashReporterCompileArgs)"/>
 			<Compile Target="CrashReportClientEditor" Platform="Linux" Configuration="Shipping" Tag="#Build Tools Linux" Arguments="$(CrashReporterCompileArgs)"/>
 			<Compile Target="ShaderCompileWorker" Configuration="Development" Platform="Linux" Tag="#Build Tools Linux"/>
+			<Compile Target="EpicWebHelper" Configuration="Shipping" Platform="Linux" Tag="#Build Tools Linux"/>
 			<Compile Target="UnrealInsights" Configuration="Development" Platform="Linux" Tag="#Build Tools Linux"/>
 			<Compile Target="UnrealLightmass" Configuration="Development" Platform="Linux" Tag="#Build Tools Linux"/>
 			<Compile Target="UnrealPak" Configuration="Development" Platform="Linux" Tag="#Build Tools Linux"/>
@@ -631,6 +632,7 @@
 			<Compile Target="CrashReportClient" Platform="Linux" Configuration="Shipping" Tag="#Build Tools LinuxNative" Arguments="$(CrashReporterCompileArgs)"/>
 			<Compile Target="CrashReportClientEditor" Platform="Linux" Configuration="Shipping"  Tag="#Build Tools LinuxNative" Arguments="$(CrashReporterCompileArgs)"/>
 			<Compile Target="ShaderCompileWorker" Configuration="Development" Platform="Linux" Tag="#Build Tools LinuxNative"/>
+			<Compile Target="EpicWebHelper" Configuration="Shipping" Platform="Linux" Tag="#Build Tools LinuxNative"/>
 			<Compile Target="UnrealInsights" Configuration="Development" Platform="Linux" Tag="#Build Tools LinuxNative"/>
 			<Compile Target="UnrealLightmass" Configuration="Development" Platform="Linux" Tag="#Build Tools LinuxNative"/>
 			<Compile Target="UnrealPak" Configuration="Development" Platform="Linux" Tag="#Build Tools LinuxNative"/>
