EpicWebHelper undecorates Shipping builds, which is what we depend on for Launching

Move to build these for Win/Mac in Shipping to replace the checked into p4 version

diff --git a/Engine/Build/InstalledEngineBuild.xml b/Engine/Build/InstalledEngineBuild.xml
--- a/Engine/Build/InstalledEngineBuild.xml
+++ b/Engine/Build/InstalledEngineBuild.xml
@@ -472,7 +472,7 @@
 			<Compile Target="CrashReportClientEditor" Configuration="Shipping" Platform="Win64" Tag="#Build Tools Win64" Arguments="$(CrashReporterCompileArgs)"/>
 			<Tag Files="#NotForLicensee Build Tools Win64" With="#Build Tools Win64"/>
 			<Compile Target="ShaderCompileWorker" Configuration="Development" Platform="Win64" Tag="#Build Tools Win64"/>
-			<Compile Target="EpicWebHelper" Configuration="Development" Platform="Win64" Tag="#Build Tools Win64"/>
+			<Compile Target="EpicWebHelper" Configuration="Shipping" Platform="Win64" Tag="#Build Tools Win64"/>
 			<Compile Target="UnrealInsights" Configuration="Development" Platform="Win64" Tag="#Build Tools Win64"/>
 			<Compile Target="UnrealFrontend" Configuration="Development" Platform="Win64" Tag="#Build Tools Win64"/>
 			<Compile Target="UnrealLightmass" Configuration="Development" Platform="Win64" Tag="#Build Tools Win64"/>
@@ -611,7 +611,7 @@
 			<Compile Target="CrashReportClient" Configuration="Shipping" Platform="Mac" Tag="#Build Tools Mac" Arguments="-EnableDSYM $(CrashReporterCompileArgs)"/>
 			<Compile Target="CrashReportClientEditor" Configuration="Shipping" Platform="Mac" Tag="#Build Tools Mac" Arguments="-EnableDSYM $(CrashReporterCompileArgs)"/>
 			<Compile Target="ShaderCompileWorker" Configuration="Development" Platform="Mac" Tag="#Build Tools Mac" Arguments="-EnableDSYM"/>
-			<Compile Target="EpicWebHelper" Configuration="Development" Platform="Mac" Tag="#Build Tools Mac" Arguments="-EnableDSYM"/>
+			<Compile Target="EpicWebHelper" Configuration="Shipping" Platform="Mac" Tag="#Build Tools Mac" Arguments="-EnableDSYM"/>
 			<Compile Target="UnrealInsights" Configuration="Development" Platform="Mac" Tag="#Build Tools Mac" Arguments="-EnableDSYM"/>
 			<Compile Target="UnrealFrontend" Configuration="Development" Platform="Mac" Tag="#Build Tools Mac" Arguments="-EnableDSYM"/>
 			<Compile Target="UnrealLightmass" Configuration="Development" Platform="Mac" Tag="#Build Tools Mac" Arguments="-EnableDSYM"/>
