Add Zen Store DDC and Lazy Load Editor Perf Tests

diff --git a/Engine/Build/EditorPerf.xml b/Engine/Build/EditorPerf.xml
--- a/Engine/Build/EditorPerf.xml
+++ b/Engine/Build/EditorPerf.xml
@@ -60,6 +60,7 @@
 	<Property Name="TraceDir" Value="$(RootDir)/Engine/Programs/UnrealInsights/Saved/TraceSessions"/>
 
 	<!-- Arguments for phases, keep in mind that these are passed to both the DDCCleanupCommandlet and RunUnreal -->
+	<Option Name="EnableBasePhaseTests" DefaultValue="true" Description="Enable the base phase perf tests."/>
 	<Option Name="ColdArgs" DefaultValue="-ddc=Cold" Description="Arguments to unreal for the cold phase"/>
 	<Option Name="IncrArgs" DefaultValue="-ddc=Warm" Description="Arguments to unreal for the incremental phase"/>
 	<Option Name="WarmArgs" DefaultValue="-ddc=Warm" Description="Arguments to unreal for the warm phase"/>
@@ -67,7 +68,7 @@
 	<!-- Iterations to run and average together -->
 	<Option Name="Iterations" DefaultValue="1" Description="The list of iteration numbers to run. For example '1;2;3' for three iterations"/>
 
-	<!-- Arguments for extra Zen Loader tests -->
+	<!-- Options for extra Zen Loader tests -->
 	<Option Name="EnableZenLoaderTest" DefaultValue="true" Description="Enable additional test per phase that uses the Zen Loader."/>
 	<Option Name="EnableZenAsyncThreadTest" DefaultValue="true" Description="Enable additional test per phase that uses the Zen Loader."/>
 	<Option Name="EnableZenTestsWithoutEditorDomain" DefaultValue="true" Description="Enable additional test per phase that uses the Zen Loader."/>
@@ -79,6 +80,23 @@
 	<Property Name="ZenLoaderTestArgs" Value="$(ZenLoaderTestArgs),-noeditordomain" If="$(EnableZenTestsWithoutEditorDomain)" />
 	<Property Name="ZenLoaderAsyncTestArgs" Value="$(ZenLoaderAsyncTestArgs),-noeditordomain" If="$(EnableZenTestsWithoutEditorDomain)" />
 
+	<!-- Options for extra Zen Store tests -->
+	<Option Name="EnableZenStoreDDCTest" DefaultValue="false" Description="Enable additional test on warm phase that uses the Zen Store DDC."/>
+
+	<!-- Zen Store tests specific arguments -->
+	<Property Name="ZenServerCache" Value="$(RootDir)\ZenCache" />
+	<Property Name="ZenStoreCommonArgs" Value="-ini:Engine:[Zen.AutoLaunch]:LimitProcessLifetime=true,-ZenDataPath=$(ZenServerCache) "/>
+	<Property Name="ZenStoreDDCPhaseArgs" Value="-ddc=ZenDDC"/>
+	<Property Name="ZenStoreDDCTestArgs" Value="$(ZenStoreCommonArgs)"/>
+
+	<!-- Options for extra LazyLoad tests -->
+	<Option Name="EnableLazyLoadTest" DefaultValue="false" Description="Enable additional test on warm phase that uses LazyLoad."/>
+	<Option Name="EnableLazyLoadZenTest" DefaultValue="false" Description="Enable additional test on warm phase that uses LazyLoad with Zen Loader."/>
+	
+	<!-- LazyLoad tests specific arguments -->
+	<Property Name="LazyLoadTestArgs" Value="-LazyLoadImports"/>
+	<Property Name="LazyLoadZenTestArgs" Value="$(LazyLoadTestArgs),-zenloader" />
+
 	<!-- Gauntlet arguments for all runs -->
 	<Property Name="GauntletArgs" Value="-build=editor -UseEditor -verbose -tempdir=&quot;$(RootDir)/Tests&quot;"/>
 	<Property Name="GauntletArgs" Value="$(GauntletArgs) -branch=$(Branch) -ECBranch=&quot;$(Branch)&quot; -ECChangelist=&quot;$(Change)&quot;"/>
@@ -191,30 +209,57 @@
 			<Property Name="RunNodeList" Value=""/> <!-- Each node in this the nested loops below, they need to depend on predecessors -->
 			<Agent Name="EditorPerf Agent $(Platform) $(Project)" Type="EditorPerf$(Platform)">
 				<ForEach Name="Phase" Values="$(Phases)">
-					<ForEach Name="Iteration" Values="$(Iterations)">
-						<!-- Only ever do one iteration of incremental, because it will warm up the workspace -->
-						<Do If="'$(Phase)' != 'Incr' or '$(Iteration)' == '1'">
-							<Property Name="IterationSuffix" Value=""/>
-							<Property Name="IterationSuffix" Value="$(Iteration)" If="'$(Iteration)' != '1'"/>
-							<Property Name="RunNodeName" Value="Run $(OutputName) $(Phase)$(IterationSuffix)"/>
-							<Node Name="$(RunNodeName)" Requires="$(RunNodeRequires);$(RunNodeList)">
-								<Expand Name="RunTest" TestNameSuffix="" Platform="$(Platform)" Project="$(Project)" NodeName="$(RunNodeName)" Phase="$(Phase)" Iteration="$(Iteration)" Args="$($(Phase)Args)" ZenArgs=""/>
-							</Node>
-							<Property Name="RunNodeList" Value="$(RunNodeList);$(RunNodeName)"/>
-						</Do>
-					</ForEach>
-					<!-- Zen Loader TestRun if enabled -->
-					<Property Name="RunNodeZenName" Value="Run $(OutputName) $(Phase) Zen"/>
-					<Node Name="$(RunNodeZenName)" Requires="$(RunNodeRequires);$(RunNodeList)" If="$(EnableZenLoaderTest)">
-						<Expand Name="RunTest" TestNameSuffix="Zen" Platform="$(Platform)" Project="$(Project)" NodeName="$(RunNodeZenName)" Phase="$(Phase)" Iteration="1" Args="$($(Phase)Args)" ZenArgs="$(ZenLoaderTestArgs)"/>
-					</Node>
-					<Property Name="RunNodeList" Value="$(RunNodeList);$(RunNodeZenName)" If="$(EnableZenLoaderTest)"/>
-					<!-- Zen Loader Async TestRun if enabled -->
-					<Property Name="RunNodeZenAsyncName" Value="Run $(OutputName) $(Phase) Zen Async"/>
-					<Node Name="$(RunNodeZenAsyncName)" Requires="$(RunNodeRequires);$(RunNodeList)" If="$(EnableZenAsyncThreadTest)">
-						<Expand Name="RunTest" TestNameSuffix="ZenAsync" Platform="$(Platform)" Project="$(Project)" NodeName="$(RunNodeZenAsyncName)" Phase="$(Phase)" Iteration="1" Args="$($(Phase)Args)" ZenArgs="$(ZenLoaderAsyncTestArgs)"/>
-					</Node>
-					<Property Name="RunNodeList" Value="$(RunNodeList);$(RunNodeZenAsyncName)" If="$(EnableZenAsyncThreadTest)"/>
+					<Do If="$(EnableBasePhaseTests)">
+						<ForEach Name="Iteration" Values="$(Iterations)">
+							<!-- Only ever do one iteration of incremental, because it will warm up the workspace -->
+							<Do If="'$(Phase)' != 'Incr' or '$(Iteration)' == '1'">
+								<Property Name="IterationSuffix" Value=""/>
+								<Property Name="IterationSuffix" Value="$(Iteration)" If="'$(Iteration)' != '1'"/>
+								<Property Name="RunNodeName" Value="Run $(OutputName) $(Phase)$(IterationSuffix)"/>
+								<Node Name="$(RunNodeName)" Requires="$(RunNodeRequires);$(RunNodeList)">
+									<Expand Name="RunTest" TestNameSuffix="" Platform="$(Platform)" Project="$(Project)" NodeName="$(RunNodeName)" Phase="$(Phase)" Iteration="$(Iteration)" Args="$($(Phase)Args)" ZenArgs=""/>
+								</Node>
+								<Property Name="RunNodeList" Value="$(RunNodeList);$(RunNodeName)"/>
+							</Do>
+						</ForEach>
+					</Do>	
+					<!-- Don't run Zen tests on incremental runs -->
+					<Do If="'$(Phase)' != 'Incr'">
+						<!-- Zen Loader TestRun if enabled -->
+						<Property Name="RunNodeZenName" Value="Run $(OutputName) $(Phase) Zen"/>
+						<Node Name="$(RunNodeZenName)" Requires="$(RunNodeRequires);$(RunNodeList)" If="$(EnableZenLoaderTest)">
+							<Expand Name="RunTest" TestNameSuffix="Zen" Platform="$(Platform)" Project="$(Project)" NodeName="$(RunNodeZenName)" Phase="$(Phase)" Iteration="1" Args="$($(Phase)Args)" ZenArgs="$(ZenLoaderTestArgs)"/>
+						</Node>
+						<Property Name="RunNodeList" Value="$(RunNodeList);$(RunNodeZenName)" If="$(EnableZenLoaderTest)"/>
+						<!-- Zen Loader Async TestRun if enabled -->
+						<Property Name="RunNodeZenAsyncName" Value="Run $(OutputName) $(Phase) Zen Async"/>
+						<Node Name="$(RunNodeZenAsyncName)" Requires="$(RunNodeRequires);$(RunNodeList)" If="$(EnableZenAsyncThreadTest)">
+							<Expand Name="RunTest" TestNameSuffix="ZenAsync" Platform="$(Platform)" Project="$(Project)" NodeName="$(RunNodeZenAsyncName)" Phase="$(Phase)" Iteration="1" Args="$($(Phase)Args)" ZenArgs="$(ZenLoaderAsyncTestArgs)"/>
+						</Node>
+						<Property Name="RunNodeList" Value="$(RunNodeList);$(RunNodeZenAsyncName)" If="$(EnableZenAsyncThreadTest)"/>
+					
+						<!-- Zen Store DDC TestRun if enabled -->
+						<Property Name="RunNodeZenStoreDDCName" Value="Run $(OutputName) $(Phase) Zen Store DDC"/>
+						<Node Name="$(RunNodeZenStoreDDCName)" Requires="$(RunNodeRequires);$(RunNodeList)" If="$(EnableZenStoreDDCTest)">
+							<Expand Name="RunTest" TestNameSuffix="ZenAsync" Platform="$(Platform)" Project="$(Project)" NodeName="$(RunNodeZenStoreDDCName)" Phase="$(Phase)" Iteration="1" Args="$(ZenStoreDDCPhaseArgs)" ZenArgs="$(ZenStoreDDCTestArgs)"/>
+						</Node>
+						<Property Name="RunNodeList" Value="$(RunNodeList);$(RunNodeZenStoreDDCName)" If="$(EnableZenStoreDDCTest)"/>
+					</Do>
+					<!-- Warm only test runs e.g LazyLoad -->
+					<Do If="'$(Phase)' == 'Warm'">
+						<!-- Legacy LazyLoad TestRun if enabled -->
+						<Property Name="RunNodeLazyLoadName" Value="Run $(OutputName) $(Phase) LazyLoad"/>
+						<Node Name="$(RunNodeLazyLoadName)" Requires="$(RunNodeRequires);$(RunNodeList)" If="$(EnableLazyLoadTest)">
+							<Expand Name="RunTest" TestNameSuffix="Zen" Platform="$(Platform)" Project="$(Project)" NodeName="$(RunNodeLazyLoadName)" Phase="$(Phase)" Iteration="1" Args="$($(Phase)Args)" ZenArgs="$(LazyLoadTestArgs)"/>
+						</Node>
+						<Property Name="RunNodeList" Value="$(RunNodeList);$(RunNodeLazyLoadName)" If="$(EnableLazyLoadTest)"/>
+						<!-- Zen LazyLoad TestRun if enabled -->
+						<Property Name="RunNodeLazyLoadZenName" Value="Run $(OutputName) $(Phase) LazyLoad Zen"/>
+						<Node Name="$(RunNodeLazyLoadZenName)" Requires="$(RunNodeRequires);$(RunNodeList)" If="$(EnableLazyLoadZenTest)">
+							<Expand Name="RunTest" TestNameSuffix="Zen" Platform="$(Platform)" Project="$(Project)" NodeName="$(RunNodeLazyLoadZenName)" Phase="$(Phase)" Iteration="1" Args="$($(Phase)Args)" ZenArgs="$(LazyLoadZenTestArgs)"/>
+						</Node>
+						<Property Name="RunNodeList" Value="$(RunNodeList);$(RunNodeLazyLoadZenName)" If="$(EnableLazyLoadZenTest)"/>
+					</Do>
 				</ForEach>
 			</Agent>
 			<Label Category="$(RunCategory)" Name="$(Project) Editor Run" Requires="$(RunNodeRequires);$(RunNodeList)" UgsBadge="$(Project) Editor Run"/>
