Addressing InstalledEngineBuild for mac requiring Windows steps

Mac DDC was requiring Win64 C# tools, split Make Feature packs to have a dedicated mac agent to remove native & C# Win64 requirements

#rnx

diff --git a/Engine/Build/InstalledEngineBuild.xml b/Engine/Build/InstalledEngineBuild.xml
--- a/Engine/Build/InstalledEngineBuild.xml
+++ b/Engine/Build/InstalledEngineBuild.xml
@@ -690,6 +690,19 @@
 		</Node>
 	</Agent>
 
+	<Agent Name="FeaturePacks Group Mac" Type="Mac_Licensee">
+		<Node Name="Make Feature Packs Mac" Requires="Update Version Files;Build Tools Mac">
+			<!--Delete any pre-existing feature packs-->
+			<Delete Files="FeaturePacks/..."/>
+			<ForEach Name="TemplateName" Values="$(ProjectsToFeaturePack)">
+				<PakFile Output="FeaturePacks/$(TemplateName).upack" Files="" ResponseFile="Templates/$(TemplateName)/contents.txt"/>
+			</ForEach>
+			<ForEach Name="SampleName" Values="StarterContent">
+				<PakFile Output="FeaturePacks/$(SampleName).upack" Files="" ResponseFile="Samples/$(SampleName)/contents.txt"/>
+			</ForEach>
+		</Node>
+	</Agent>
+
 	<Agent Name="FeaturePacks Group LinuxNative" Type="Linux_Licensee">
 		<Node Name="Make Feature Packs LinuxNative" Requires="Update Version Files;Build Tools LinuxNative">
 			<!--Delete any pre-existing feature packs-->
@@ -798,7 +811,7 @@
 	<Property Name="MacDDCAgentType" Value="Mac" />
 	<Property Name="MacDDCAgentType" Value="MacArm64" If="$(WithMacArm64)" />
 	<Agent Name="DDC Group Mac" Type="$(MacDDCAgentType)">
-		<Node Name="Build DDC Mac" Requires="#UnrealEditor Mac;Build Tools Mac;#Build Tools CS Binaries" Produces="#CompressedDDCMac">
+		<Node Name="Build DDC Mac" Requires="#UnrealEditor Mac;Build Tools Mac;#Build Tools CS Binaries Mac" Produces="#CompressedDDCMac">
 			<!-- Get our temp dir -->
 			<Property Name="DDCDir" Value="$(RootDir)/LocalBuilds/InstalledDDC"/>
 			<Property Name="SavedDir" Value="$(SavedOutput)/Installed/Mac"/>
@@ -820,7 +833,7 @@
 			
 			<!-- Copy everything to a temporary directory -->
 			<Copy From="$(RootDir)" To="$(DDCDir)" Files="#FilteredCopyList"/>
-			<Copy From="$(CsToolsDir)" To="$(DDCDir)" Files="#Build Tools CS Binaries"/>
+			<Copy From="$(CsToolsDir)" To="$(DDCDir)" Files="#Build Tools CS Binaries Mac"/>
 
 			<Command Name="BuildDerivedDataCache" Arguments="-TempDir=&quot;$(DDCDir)&quot; -FeaturePacks=&quot;$(ProjectsToBuildDDC)&quot; -TargetPlatforms=$(DDCPlatformsMac) -HostPlatform=Mac -SavedDir=&quot;$(SavedDir)&quot;"/>
 
@@ -1068,7 +1081,7 @@
 		<Property Name="FilterRequirements" Value="$(FilterRequirements);#Compile Datasmith Plugins Mac;#DatasmithPluginBinaries Mac" If="$(WithMac) And $(CompileDatasmithPlugins)"/>
 
 		<!-- Feature packs -->
-		<Property Name="FilterRequirements" Value="$(FilterRequirements);Make Feature Packs"/>
+		<Property Name="FilterRequirements" Value="$(FilterRequirements);Make Feature Packs Mac"/>
 
 		<!-- Filters the build products required for an installed build into groups depending on whether they need to be stripped, signed or just copied-->
 		<Node Name="Make Installed Build Mac" Requires="$(FilterRequirements)" Produces="#Installed Build Mac Files">
@@ -1086,7 +1099,7 @@
 			<Tag Files="#Compile UnrealHeaderTool Mac" With="#Installed Mac"/>
 			<Tag Files="$(UnrealEditor)" Except="#UnrealEditor Mac Unstripped;#UnrealEditor Mac Unsigned" With="#Installed Mac"/>
 			<Tag Files="#UnrealEditor Mac Stripped;#UnrealEditor Mac Signed" With="#Saved Output"/>
-			<Tag Files="#Build Tools Mac;#Make Feature Packs;#Rules Assemblies" With="#Installed Mac"/>
+			<Tag Files="#Build Tools Mac;#Make Feature Packs Mac;#Rules Assemblies" With="#Installed Mac"/>
 			<Tag Files="#DatasmithPluginBinaries Mac" With="#Installed Mac" If="$(WithMac) And $(CompileDatasmithPlugins)"/>
 
 			<!-- Tag enabled target platforms -->
