Add nanosvg to the EditorFilter allowing the source files to be included in an InstalledEngineBuild

the Build.cs and tps files were, but not the actual source.

#rnx

diff --git a/Engine/Build/InstalledEngineFilters.xml b/Engine/Build/InstalledEngineFilters.xml
--- a/Engine/Build/InstalledEngineFilters.xml
+++ b/Engine/Build/InstalledEngineFilters.xml
@@ -298,6 +298,9 @@
 
 		<!-- Used by Switchboard application, and by iOS/tvOS build automation -->
 		Engine/Extras/ThirdPartyNotUE/cwrsync/...
+
+		<!-- nanosvg -->
+		Engine/Source/ThirdParty/nanosvg/...
 	</Property>
 
 	<Property Name="CopyEditorExceptionsWin64" Value="$(CopyEditorExceptions)">
