Make sure InterchangeWorker is build by UGS and when we release it.

#rnx

diff --git a/Engine/Build/BatchFiles/Linux/README.md b/Engine/Build/BatchFiles/Linux/README.md
--- a/Engine/Build/BatchFiles/Linux/README.md
+++ b/Engine/Build/BatchFiles/Linux/README.md
@@ -93,7 +93,7 @@ being the default. You can now just type
 
 to build the editor. Alternatively, build the following targets:
 
-    make CrashReportClient ShaderCompileWorker UnrealLightmass UnrealPak UnrealEditor
+    make CrashReportClient ShaderCompileWorker UnrealLightmass InterchangeWorker UnrealPak UnrealEditor
 
 If you intend to develop the editor, you can build a debug configuration of it:
 
diff --git a/Engine/Build/Graph/Examples/BuildEditorAndTools.xml b/Engine/Build/Graph/Examples/BuildEditorAndTools.xml
--- a/Engine/Build/Graph/Examples/BuildEditorAndTools.xml
+++ b/Engine/Build/Graph/Examples/BuildEditorAndTools.xml
@@ -65,6 +65,7 @@
 			<Compile Target="UnrealHeaderTool" Platform="Win64" Configuration="Development" Tag="#ToolBinaries"/>
 			<Compile Target="ShaderCompileWorker" Platform="Win64" Configuration="Development" Tag="#ToolBinaries"/>
 			<Compile Target="UnrealLightmass" Platform="Win64" Configuration="Development" Tag="#ToolBinaries"/>
+			<Compile Target="InterchangeWorker" Platform="Win64" Configuration="Development" Tag="#ToolBinaries"/>
 			<Compile Target="UnrealPak" Configuration="Development" Platform="Win64" Tag="#ToolBinaries"/>
 			<ForEach Name="AdditonalTool" Values="$(AdditionalTools)">
 				<Compile Target="$(AdditonalTool)" Platform="Win64" Configuration="Development" Tag="#ToolBinaries"/>
diff --git a/Engine/Build/InstalledEngineBuild.xml b/Engine/Build/InstalledEngineBuild.xml
--- a/Engine/Build/InstalledEngineBuild.xml
+++ b/Engine/Build/InstalledEngineBuild.xml
@@ -476,6 +476,7 @@
 			<Compile Target="UnrealInsights" Configuration="Development" Platform="Win64" Tag="#Build Tools Win64"/>
 			<Compile Target="UnrealFrontend" Configuration="Development" Platform="Win64" Tag="#Build Tools Win64"/>
 			<Compile Target="UnrealLightmass" Configuration="Development" Platform="Win64" Tag="#Build Tools Win64"/>
+			<Compile Target="InterchangeWorker" Configuration="Development" Platform="Win64" Tag="#Build Tools Win64"/>
 			<Compile Target="UnrealPak" Configuration="Development" Platform="Win64" Tag="#Build Tools Win64"/>
 			<Compile Target="UnrealMultiUserServer" Configuration="Development" Platform="Win64" Tag="#Build Tools Win64"/>
 			<Compile Target="UnrealMultiUserSlateServer" Configuration="Development" Platform="Win64" Tag="#Build Tools Win64"/>
@@ -550,6 +551,7 @@
 			<Compile Target="EpicWebHelper" Configuration="Shipping" Platform="Linux" Tag="#Build Tools Linux"/>
 			<Compile Target="UnrealInsights" Configuration="Development" Platform="Linux" Tag="#Build Tools Linux"/>
 			<Compile Target="UnrealLightmass" Configuration="Development" Platform="Linux" Tag="#Build Tools Linux"/>
+			<Compile Target="InterchangeWorker" Configuration="Development" Platform="Linux" Tag="#Build Tools Linux"/>
 			<Compile Target="UnrealPak" Configuration="Development" Platform="Linux" Tag="#Build Tools Linux"/>
 			<Compile Target="UnrealMultiUserServer" Configuration="Development" Platform="Linux" Tag="#Build Tools Linux"/>
 			<Compile Target="UnrealMultiUserSlateServer" Configuration="Development" Platform="Linux" Tag="#Build Tools Linux"/>
@@ -615,6 +617,7 @@
 			<Compile Target="UnrealInsights" Configuration="Development" Platform="Mac" Tag="#Build Tools Mac" Arguments="-EnableDSYM"/>
 			<Compile Target="UnrealFrontend" Configuration="Development" Platform="Mac" Tag="#Build Tools Mac" Arguments="-EnableDSYM"/>
 			<Compile Target="UnrealLightmass" Configuration="Development" Platform="Mac" Tag="#Build Tools Mac" Arguments="-EnableDSYM"/>
+			<Compile Target="InterchangeWorker" Configuration="Development" Platform="Mac" Tag="#Build Tools Mac" Arguments="-EnableDSYM"/>
 			<Compile Target="UnrealPak" Configuration="Development" Platform="Mac" Tag="#Build Tools Mac" Arguments="-EnableDSYM"/>
 			<Compile Target="UnrealMultiUserServer" Configuration="Development" Platform="Mac" Tag="#Build Tools Mac" Arguments="-EnableDSYM"/>
 			<Compile Target="UnrealMultiUserSlateServer" Configuration="Development" Platform="Mac" Tag="#Build Tools Mac" Arguments="-EnableDSYM"/>
@@ -635,6 +638,7 @@
 			<Compile Target="EpicWebHelper" Configuration="Shipping" Platform="Linux" Tag="#Build Tools LinuxNative"/>
 			<Compile Target="UnrealInsights" Configuration="Development" Platform="Linux" Tag="#Build Tools LinuxNative"/>
 			<Compile Target="UnrealLightmass" Configuration="Development" Platform="Linux" Tag="#Build Tools LinuxNative"/>
+			<Compile Target="InterchangeWorker" Configuration="Development" Platform="Linux" Tag="#Build Tools LinuxNative"/>
 			<Compile Target="UnrealPak" Configuration="Development" Platform="Linux" Tag="#Build Tools LinuxNative"/>
 			<Compile Target="UnrealMultiUserServer" Configuration="Development" Platform="Linux" Tag="#Build Tools LinuxNative"/>
 			<Compile Target="UnrealMultiUserSlateServer" Configuration="Development" Platform="Linux" Tag="#Build Tools LinuxNative"/>
diff --git a/Engine/Build/Rsync/RsyncEngine.txt b/Engine/Build/Rsync/RsyncEngine.txt
--- a/Engine/Build/Rsync/RsyncEngine.txt
+++ b/Engine/Build/Rsync/RsyncEngine.txt
@@ -162,6 +162,7 @@ protect /Saved/
 + /Source/Programs/UnrealHeaderTool/
 + /Source/Programs/UnrealHeaderTool/**.def
 + /Source/Programs/UnrealLightmass/
++ /Source/Programs/InterchangeWorker/
 - /Source/Programs/**/bin/
 - /Source/Programs/**/obj/
 - /Source/Programs/*/
diff --git a/Engine/Source/Programs/AutomationTool/BuildGraph/Scripts/InstalledBuild.cs b/Engine/Source/Programs/AutomationTool/BuildGraph/Scripts/InstalledBuild.cs
--- a/Engine/Source/Programs/AutomationTool/BuildGraph/Scripts/InstalledBuild.cs
+++ b/Engine/Source/Programs/AutomationTool/BuildGraph/Scripts/InstalledBuild.cs
@@ -327,6 +327,7 @@ namespace AutomationTool
 			Files += await CompileAsync("UnrealInsights", UnrealTargetPlatform.Win64, UnrealTargetConfiguration.Development);
 			Files += await CompileAsync("UnrealFrontend", UnrealTargetPlatform.Win64, UnrealTargetConfiguration.Development);
 			Files += await CompileAsync("UnrealLightmass", UnrealTargetPlatform.Win64, UnrealTargetConfiguration.Development);
+			Files += await CompileAsync("InterchangeWorker", UnrealTargetPlatform.Win64, UnrealTargetConfiguration.Development);
 			Files += await CompileAsync("UnrealPak", UnrealTargetPlatform.Win64, UnrealTargetConfiguration.Development);
 			Files += await CompileAsync("UnrealMultiUserServer", UnrealTargetPlatform.Win64, UnrealTargetConfiguration.Development);
 			Files += await CompileAsync("UnrealRecoverySvc", UnrealTargetPlatform.Win64, UnrealTargetConfiguration.Development);
diff --git a/Engine/Source/Programs/AutomationTool/Scripts/BuildCommonTools.Automation.cs b/Engine/Source/Programs/AutomationTool/Scripts/BuildCommonTools.Automation.cs
--- a/Engine/Source/Programs/AutomationTool/Scripts/BuildCommonTools.Automation.cs
+++ b/Engine/Source/Programs/AutomationTool/Scripts/BuildCommonTools.Automation.cs
@@ -94,6 +94,7 @@ public class BuildCommonTools : BuildCommand
 			Agenda.AddTarget("UnrealHeaderTool", UnrealBuildTool.UnrealTargetPlatform.Win64, UnrealBuildTool.UnrealTargetConfiguration.Development);
 			Agenda.AddTarget("UnrealPak", UnrealBuildTool.UnrealTargetPlatform.Win64, UnrealBuildTool.UnrealTargetConfiguration.Development);
 			Agenda.AddTarget("UnrealLightmass", UnrealBuildTool.UnrealTargetPlatform.Win64, UnrealBuildTool.UnrealTargetConfiguration.Development);
+			Agenda.AddTarget("InterchangeWorker", UnrealBuildTool.UnrealTargetPlatform.Win64, UnrealBuildTool.UnrealTargetConfiguration.Development);
 			Agenda.AddTarget("ShaderCompileWorker", UnrealBuildTool.UnrealTargetPlatform.Win64, UnrealBuildTool.UnrealTargetConfiguration.Development);
 			Agenda.AddTarget("UnrealVersionSelector", UnrealBuildTool.UnrealTargetPlatform.Win64, UnrealBuildTool.UnrealTargetConfiguration.Shipping);
 			Agenda.AddTarget("BootstrapPackagedGame", UnrealBuildTool.UnrealTargetPlatform.Win64, UnrealBuildTool.UnrealTargetConfiguration.Shipping);
@@ -106,6 +107,7 @@ public class BuildCommonTools : BuildCommand
 			Agenda.AddTarget("CrashReportClient", UnrealBuildTool.UnrealTargetPlatform.Mac, UnrealBuildTool.UnrealTargetConfiguration.Development, InAddArgs: "-CopyAppBundleBackToDevice");
 			Agenda.AddTarget("UnrealPak", UnrealBuildTool.UnrealTargetPlatform.Mac, UnrealBuildTool.UnrealTargetConfiguration.Development, InAddArgs: "-CopyAppBundleBackToDevice");
 			Agenda.AddTarget("UnrealLightmass", UnrealBuildTool.UnrealTargetPlatform.Mac, UnrealBuildTool.UnrealTargetConfiguration.Development, InAddArgs: "-CopyAppBundleBackToDevice");
+			Agenda.AddTarget("InterchangeWorker", UnrealBuildTool.UnrealTargetPlatform.Mac, UnrealBuildTool.UnrealTargetConfiguration.Development, InAddArgs: "-CopyAppBundleBackToDevice");
 			Agenda.AddTarget("ShaderCompileWorker", UnrealBuildTool.UnrealTargetPlatform.Mac, UnrealBuildTool.UnrealTargetConfiguration.Development, InAddArgs: "-CopyAppBundleBackToDevice");
 			Agenda.AddTarget("UnrealEditorServices", UnrealBuildTool.UnrealTargetPlatform.Mac, UnrealBuildTool.UnrealTargetConfiguration.Development, InAddArgs: "-CopyAppBundleBackToDevice");
 			Agenda.AddTarget("EpicWebHelper", UnrealBuildTool.UnrealTargetPlatform.Mac, UnrealBuildTool.UnrealTargetConfiguration.Development, InAddArgs: "-CopyAppBundleBackToDevice");
diff --git a/Engine/Source/Programs/AutomationTool/Scripts/BuildTarget.Automation.cs b/Engine/Source/Programs/AutomationTool/Scripts/BuildTarget.Automation.cs
--- a/Engine/Source/Programs/AutomationTool/Scripts/BuildTarget.Automation.cs
+++ b/Engine/Source/Programs/AutomationTool/Scripts/BuildTarget.Automation.cs
@@ -157,6 +157,7 @@ namespace AutomationTool
 						Agenda.AddTarget("UnrealPak", CurrentPlatform, UnrealTargetConfiguration.Development, ProjectFile, UBTArgs);
 						Agenda.AddTarget("ShaderCompileWorker", CurrentPlatform, UnrealTargetConfiguration.Development, ProjectFile, UBTArgs);
 						Agenda.AddTarget("UnrealLightmass", CurrentPlatform, UnrealTargetConfiguration.Development, ProjectFile, UBTArgs);
+						Agenda.AddTarget("InterchangeWorker", CurrentPlatform, UnrealTargetConfiguration.Development, ProjectFile, UBTArgs);
 						Agenda.AddTarget("CrashReportClient", CurrentPlatform, UnrealTargetConfiguration.Shipping, ProjectFile, UBTArgs);
 						Agenda.AddTarget("CrashReportClientEditor", CurrentPlatform, UnrealTargetConfiguration.Shipping, ProjectFile, UBTArgs);
 					}
diff --git a/Engine/Source/Programs/AutomationTool/Scripts/CopySharedCookedBuild.Automation.cs b/Engine/Source/Programs/AutomationTool/Scripts/CopySharedCookedBuild.Automation.cs
--- a/Engine/Source/Programs/AutomationTool/Scripts/CopySharedCookedBuild.Automation.cs
+++ b/Engine/Source/Programs/AutomationTool/Scripts/CopySharedCookedBuild.Automation.cs
@@ -133,6 +133,7 @@ class CopySharedCookedBuild : BuildCommand
 				Agenda.AddTarget(EditorTargetName, CurrentPlatform, UnrealTargetConfiguration.Development, ProjectFileName.EndsWith(".uproject", StringComparison.InvariantCultureIgnoreCase)? new FileReference(ProjectFileName) : null);
 				Agenda.AddTarget("ShaderCompileWorker", CurrentPlatform, UnrealTargetConfiguration.Development);
 				Agenda.AddTarget("UnrealLightmass", CurrentPlatform, UnrealTargetConfiguration.Development);
+				Agenda.AddTarget("InterchangeWorker", CurrentPlatform, UnrealTargetConfiguration.Development);
 				Agenda.AddTarget("CrashReportClient", CurrentPlatform, UnrealTargetConfiguration.Shipping);
 
 				UnrealBuild Build = new UnrealBuild(this);
diff --git a/Engine/Source/Programs/InterchangeWorker/InterchangeWorker.Target.cs b/Engine/Source/Programs/InterchangeWorker/InterchangeWorker.Target.cs
--- a/Engine/Source/Programs/InterchangeWorker/InterchangeWorker.Target.cs
+++ b/Engine/Source/Programs/InterchangeWorker/InterchangeWorker.Target.cs
@@ -4,6 +4,7 @@ using UnrealBuildTool;
 using System.Collections.Generic;
 
 [SupportedPlatforms(UnrealPlatformClass.Desktop)]
+[SupportedConfigurations(UnrealTargetConfiguration.Debug, UnrealTargetConfiguration.Development)]
 public class InterchangeWorkerTarget : TargetRules
 {
 	public InterchangeWorkerTarget(TargetInfo Target) : base(Target)
diff --git a/Engine/Source/Programs/SlateUGS/Private/UGSTab.cpp b/Engine/Source/Programs/SlateUGS/Private/UGSTab.cpp
--- a/Engine/Source/Programs/SlateUGS/Private/UGSTab.cpp
+++ b/Engine/Source/Programs/SlateUGS/Private/UGSTab.cpp
@@ -242,6 +242,7 @@ TMap<FGuid, UGSCore::FCustomConfigObject> UGSTab::GetDefaultBuildStepObjects(con
 	DefaultBuildSteps.Add(UGSCore::FBuildStep(FGuid(0x91C2A429, 0xC39149B4, 0x92A54E6B, 0xE71E0F00), 9, TEXT("Compile SwitchboardListener"), TEXT("Compiling SwitchboardListener..."), 1, TEXT("SwitchboardListener"), HostPlatform, TEXT("Development"), TEXT(""), !ShouldSyncPrecompiledEditor()));
 	DefaultBuildSteps.Add(UGSCore::FBuildStep(FGuid(0x5036C75B, 0x8DF04329, 0x82A1869D, 0xD2D48605), 10, TEXT("Compile UnrealMultiUserServer"), TEXT("Compiling UnrealMultiUserServer..."), 1, TEXT("UnrealMultiUserServer"), HostPlatform, TEXT("Development"), TEXT(""), !ShouldSyncPrecompiledEditor()));
 	DefaultBuildSteps.Add(UGSCore::FBuildStep(FGuid(0x274B89C3, 0x9DC64465, 0xA50840AB, 0xC4593CC2), 11, TEXT("Compile UnrealMultiUserSlateServer"), TEXT("Compiling UnrealMultiUserSlateServer..."), 1, TEXT("UnrealMultiUserSlateServer"), HostPlatform, TEXT("Development"), TEXT(""), !ShouldSyncPrecompiledEditor()));
+	DefaultBuildSteps.Add(UGSCore::FBuildStep(FGuid(0xF7F0C4C7, 0x55514485, 0xA961F3E4, 0xFD9D138C), 12, TEXT("Compile InterchangeWorker"), TEXT("Compiling InterchangeWorker..."), 1, TEXT("InterchangeWorker"), HostPlatform, TEXT("Development"), TEXT(""), !ShouldSyncPrecompiledEditor()));
 
 	TMap<FGuid, UGSCore::FCustomConfigObject> DefaultBuildStepObjects;
 	for (const UGSCore::FBuildStep& DefaultBuildStep : DefaultBuildSteps)
diff --git a/Engine/Source/Programs/UGSCore/Workspace.cpp b/Engine/Source/Programs/UGSCore/Workspace.cpp
--- a/Engine/Source/Programs/UGSCore/Workspace.cpp
+++ b/Engine/Source/Programs/UGSCore/Workspace.cpp
@@ -23,6 +23,7 @@ const TCHAR* FWorkspace::DefaultBuildTargets[] =
 	TEXT("ShaderCompileWorker Win64 Development, 0.8"),
 	TEXT("UnrealLightmass Win64 Development, 0.9"),
 	TEXT("CrashReportClient Win64 Shipping, 1.0"),
+	TEXT("InterchangeWorker Win64 Development, 1.0"),
 };
 
 const FWorkspaceSyncCategory FWorkspace::DefaultSyncCategories[] =
