Add option to run all compiles on a single machine
Add option to run compiles with nodebuginfo
Remove UHT
#rnx

diff --git a/Engine/Build/EditorPerf.xml b/Engine/Build/EditorPerf.xml
--- a/Engine/Build/EditorPerf.xml
+++ b/Engine/Build/EditorPerf.xml
@@ -14,6 +14,7 @@
 	<Option Name="CompileCategory" DefaultValue="Editors" Description="Category string for Compile Label"/>
 	<Option Name="RunCategory" DefaultValue="Test" Description="Category string for Run Label"/>
 	<Option Name="BaselineCategory" DefaultValue="Test" Description="Category string for Baseline Label"/>
+	<Option Name="SingleAgentCompile" Restrict="true|false" DefaultValue="false" Description="Run all compile steps on same agent"/>
 
 	<!-- Projects setup-->
 	<Option Name="Projects" DefaultValue="false" Restrict=".+" Description="The projects to test"/>
@@ -26,6 +27,7 @@
 	<Error Message="No platforms specified" If="'$(Platforms)' == ''"/>
 
 	<!-- Build Options -->
+	<Option Name="NoDebugInfo" Restrict="true|false" DefaultValue="false" Description="Compile without including debug info"/>
 	<Property Name="ExtraPluginsCommaDelims" Value="EditorTests"/>
 	<Property Name="ExtraPluginsPlusDelims" Value=""/> <!-- populated programmatically -->
 	<ForEach Name="Plugin" Values="$(ExtraPluginsCommaDelims)">
@@ -35,6 +37,7 @@
 	</ForEach>
 	<Property Name="EditorCompileArgs" Value=""/>
 	<Property Name="EditorCompileArgs" Value="-BuildPlugin=$(ExtraPluginsPlusDelims)" If="'$(ExtraPluginsPlusDelims)' != ''"/>
+	<Property Name="EditorCompileArgs" Value="$(EditorCompileArgs) -nodebuginfo" If="$(NoDebugInfo)"/>
 
 	<!-- Automation uses its own network output directory -->
 	<Property Name="AutomationNetworkRootDir" Value="\\epicgames.net\root\Builds\Automation\Reports\EditorPerf"/>
@@ -192,11 +195,11 @@
 
 			<!-- Compile the Editor -->
 			<Property Name="CompileNodeName" Value="Compile $(Project) $(Platform)"/>
-			<Agent Name="Compile Editor Agent $(Platform) $(Project)" Type="$($(Platform)CompileAgentType)">
+			<Property Name="AgentModifier" Value=""/>
+			<Property Name="AgentModifier" Value=" $(Project)" If="!$(SingleAgentCompile)"/>
+			<Agent Name="Compile Editor Agent $(Platform)$(AgentModifier" Type="$($(Platform)CompileAgentType)">
 				<Node Name="$(CompileNodeName)">
 					<SetVersion Change="$(Change)" Branch="$(EscapedBranch)" If="$(IsBuildMachine)"/>
-					<Compile Target="UnrealHeaderTool" Platform="$(Platform)" Configuration="Development"/>
-					<Compile Target="UnrealHeaderTool" Platform="$(Platform)" Project="$(Project)" Configuration="Development"/>
 					<Compile Target="ShaderCompileWorker" Platform="$(Platform)" Configuration="Development"/>
 					<Compile Target="CrashReportClientEditor" Platform="$(Platform)" Project="$(Project)" Configuration="Shipping"/>
 					<Compile Target="$(EditorTarget)" Platform="$(Platform)" Project="$(Project)" Configuration="Development" Arguments="$(EditorCompileArgs)"/>
