[CompileShadersCommandlet] - Make the "CommandletArguments" an Option so it can be set from a UI.

- Provide a default value and a description.

diff --git a/Engine/Build/CompileShadersTest.xml b/Engine/Build/CompileShadersTest.xml
--- a/Engine/Build/CompileShadersTest.xml
+++ b/Engine/Build/CompileShadersTest.xml
@@ -20,7 +20,7 @@
 		<Property Name="AnalyticsArg" Value="-analytics"/>
 	</Do>
 	
-	<Property Name="CommandletArguments" Value="-targetplatform=Windows -Collection=Audit_InCook -noshaderddc"/>
+	<Option Name="CommandletArguments" Description="Arguments to pass to the CompileShaders commandlet." DefaultValue="-targetplatform=Windows -Collection=Audit_InCook -noshaderddc"/>
 	<Property Name="PrivateCommandletArguments" Value="$(CommandletArguments)  $(AnalyticsArg)"/>
 	
 	<!-- Projects setup-->
