Ensure all approprate files are digitally signed

#rnx

diff --git a/Engine/Build/InstalledEngineBuild.xml b/Engine/Build/InstalledEngineBuild.xml
--- a/Engine/Build/InstalledEngineBuild.xml
+++ b/Engine/Build/InstalledEngineBuild.xml
@@ -130,6 +130,7 @@
 	<Macro Name="Platform_Append_DDCPlatformsWin64"/>
 	<Macro Name="Platform_TagDDCFiles"/>
 	<Macro Name="Platform_Tag_UnrealGame"/>
+	<Macro Name="Platform_ExtraSignFiles"/>
 
 	<!-- Include any scripts from platform extensions -->
 	<Property Name="ExtensionPlatforms" Value=""/>
@@ -955,6 +956,7 @@
 			<!-- Define filter and exception properties for the installed build -->
 			<Property Name="CopyInstalledFilter" Value="$(CopyEditorFilterWin64)"/>
 			<Property Name="CopyInstalledExceptions" Value="$(CopyEditorExceptionsWin64)"/>
+			<Property Name="ExtraSignFiles" Value="$(ExtraSignFilesWin64)"/>
 
 			<!-- Tag files that always need to be part of an installed build -->
 			<Tag Files="#Compile UnrealHeaderTool Win64" With="#Installed Win64"/>
@@ -977,6 +979,7 @@
 				<Property Name="Platforms" Value="$(Platforms)Android;"/>
 				<Property Name="CopyInstalledFilter" Value="$(CopyInstalledFilter);$(CopyAndroidFilter)"/>
 				<Property Name="CopyInstalledExceptions" Value="$(CopyInstalledExceptions);$(CopyAndroidExceptions)"/>
+				<Property Name="ExtraSignFiles" Value="$(ExtraSignFiles);$(ExtraSignFilesAndroid)"/>
 			</Do>
 			<Do If="'$(WithIOS)' == true">
 				<Tag Files="#UnrealGame IOS;#Compile UnrealHeaderTool Mac" Except="#UnrealGame IOS Unstripped" With="#Installed Win64"/>
@@ -998,6 +1001,7 @@
 				<Property Name="Platforms" Value="$(Platforms)Linux;"/>
 				<Property Name="CopyInstalledFilter" Value="$(CopyInstalledFilter);$(CopyLinuxFilterWin64)"/>
 				<Property Name="CopyInstalledExceptions" Value="$(CopyInstalledExceptions);$(CopyLinuxExceptions)"/>
+				<Property Name="ExtraSignFiles" Value="$(ExtraSignFiles);$(ExtraSignFilesLinux)"/>
 			</Do>
 			<Do If="'$(WithLinuxArm64)' == true">
 				<Tag Files="#UnrealGame LinuxArm64" Except="#UnrealGame LinuxArm64 Unstripped" With="#Installed Win64"/>
@@ -1005,8 +1009,11 @@
 				<Property Name="Platforms" Value="$(Platforms)LinuxArm64;"/>
 				<Property Name="CopyInstalledFilter" Value="$(CopyInstalledFilter);$(CopyLinuxFilter)"/>
 				<Property Name="CopyInstalledExceptions" Value="$(CopyInstalledExceptions);$(CopyLinuxExceptions)"/>
+				<Property Name="ExtraSignFiles" Value="$(ExtraSignFiles);$(ExtraSignFilesLinux)"/>
 			</Do>
 			<Expand Name="Platform_Tag_UnrealGame"/>
+			<Expand Name="Platform_ExtraSignFiles"/>
+
 			<!-- Tag any dependencies from all previous build products -->
 			<Tag Files="#Installed Win64" Filter="*.target" With="#TargetReceipts"/>
 			<TagReceipt Files="#TargetReceipts" RuntimeDependencies="true" With="#Installed Win64"/>
@@ -1026,7 +1033,10 @@
 
 			<Do If="'$(SignExecutables)' == true">
 				<!-- Tag files that need to be signed -->
-				<Tag Files="#Compile UnrealHeaderTool Win64;#Build Tools Win64" Filter="$(WindowsSignFilter)" Except="$(CopyInstalledExceptions);$(SignExceptions)" With="#Sign Tools for Installed Win64"/>
+				<Tag Files="#Compile UnrealHeaderTool Win64;#Build Tools Win64;#Rules Assemblies" Filter="$(WindowsSignFilter)" Except="$(CopyInstalledExceptions);$(SignExceptions)" With="#Sign Tools for Installed Win64"/>
+				<Tag Files="$(ExtraSignFiles)" Filter="$(WindowsSignFilter)" Except="$(CopyInstalledExceptions);$(SignExceptions)" With="#Sign Tools for Installed Win64"/>
+				<Tag Files="#DatasmithPluginBinaries Win64" Filter="$(WindowsSignFilter)" Except="$(CopyInstalledExceptions);$(SignExceptions)" With="#Sign Tools for Installed Win64" If="$(CompileDatasmithPlugins)"/>
+				<Tag Files="#Build Tools CS Binaries;" Filter="$(WindowsSignFilter)" Except="$(CopyWin64CsToolsExceptions);$(SignExceptions);" With="#Sign CS Tools for Installed Win64"/>
 			</Do>
 
 			<!-- Tag files that need to be copied directly -->
@@ -1045,6 +1055,7 @@
 			<Do If="'$(SignExecutables)' == true">
 				<!-- Copy then Sign required files in their final location -->
 				<Copy Files="#Sign Tools for Installed Win64" From="$(RootDir)" To="$(LocalInstalledDir)" Tag="#Files to Sign"/>
+				<Copy Files="#Sign CS Tools for Installed Win64" From="$(CsToolsDir)" To="$(LocalInstalledDir)" Tag="#Files to Sign"/>
 				<Sign Files="#Files to Sign" Tag="#Installed Build Win64 Files"/>
 			</Do>
 
diff --git a/Engine/Build/InstalledEngineFilters.xml b/Engine/Build/InstalledEngineFilters.xml
--- a/Engine/Build/InstalledEngineFilters.xml
+++ b/Engine/Build/InstalledEngineFilters.xml
@@ -248,6 +248,15 @@
 
 
 	<!-- Win64 editor specific dependencies -->
+	<Property Name="ExtraSignFilesWin64">
+		<!-- CsvTools binaries which are manually updated in P4 -->
+		Engine/Binaries/DotNET/CsvTools/...
+		Engine/Binaries/ThirdParty/svn/Win64/...
+		Engine/Binaries/Win64/UnrealTraceServer.exe
+		Engine/Binaries/Win64/EpicWebHelper.exe
+		Engine/Build/Windows/cl-filter/...
+		Engine/Extras/ThirdPartyNotUE/cwrsync/bin/...
+	</Property>
 	
 	<Property Name="CopyEditorFilterWin64" Value="$(CopyEditorFilter)">
 		<!-- Build batch files -->
@@ -463,6 +472,10 @@
 	</Property>
 
 	<!-- Android -->
+	<Property Name="ExtraSignFilesAndroid">
+		Engine/Binaries/DotNET/Android/...
+	</Property>
+
 	<Property Name="CopyAndroidFilter">
 		<!-- Build Files -->
 		Engine/Binaries/DotNET/Android/...
@@ -555,6 +568,12 @@
 	</Property>
 
 	<!-- Linux -->
+	<Property Name="ExtraSignFilesLinux">
+		<!-- Linux cross build Debug tools -->
+		Engine/Binaries/Linux/dump_syms.exe
+		Engine/Binaries/Linux/BreakpadSymbolEncoder.exe
+	</Property>
+
 	<Property Name="CopyLinuxFilter">
 		Engine/Binaries/ThirdParty/OpenAL/...
 		Engine/Binaries/ThirdParty/ICU/icu4c-53_1/Unix/x86_64-unknown-linux-gnu/*.so
diff --git a/Engine/Platforms/Hololens/Build/InstalledEngineBuild.xml b/Engine/Platforms/Hololens/Build/InstalledEngineBuild.xml
--- a/Engine/Platforms/Hololens/Build/InstalledEngineBuild.xml
+++ b/Engine/Platforms/Hololens/Build/InstalledEngineBuild.xml
@@ -31,13 +31,30 @@
 
 	<Extend Name="Platform_Agent_Target_Win64_Licensee">
 		<Agent Name="Target HoloLens" Type="Win64_Licensee">
-			<Node Name="Compile UnrealGame HoloLens" Requires="Compile UnrealHeaderTool Win64" Produces="#UnrealGame HoloLens">
+			<Node Name="Compile UnrealGame HoloLens" Requires="Compile UnrealHeaderTool Win64" Produces="#UnrealGame HoloLens;#UnrealGame HoloLens Unstripped;#UnrealGame HoloLens Stripped;#UnrealGame HoloLens Unsigned;#UnrealGame HoloLens Signed">
 				<ForEach Name="Configuration" Values="$(GameConfigurations)">
 					<Compile Target="UnrealGame" Platform="HoloLens" Configuration="$(Configuration)" Tag="#UnrealGame HoloLens" Arguments="-precompile -allmodules -nolink -architecture=arm64 $(TargetDebugInfoArg)"/>
 					<Compile Target="UnrealGame" Platform="HoloLens" Configuration="$(Configuration)" Tag="#UnrealGame HoloLens" Arguments="-precompile -architecture=arm64 $(TargetDebugInfoArg)" Clean="false"/>
 					<Compile Target="UnrealGame" Platform="HoloLens" Configuration="$(Configuration)" Tag="#UnrealGame HoloLens" Arguments="-precompile -allmodules -nolink -architecture=x64 $(TargetDebugInfoArg)"/>
 					<Compile Target="UnrealGame" Platform="HoloLens" Configuration="$(Configuration)" Tag="#UnrealGame HoloLens" Arguments="-precompile -architecture=x64 $(TargetDebugInfoArg)" Clean="false"/>
 				</ForEach>
+
+				<Switch>
+					<Case If="'$(WithFullDebugInfo)' == true">
+						<!-- Copy files for local debug builds -->
+						<Copy Files="#UnrealGame HoloLens Unstripped" From="$(RootDir)" To="$(SavedOutput)"/>
+					</Case>
+					<Default>
+						<!-- Otherwise, strip the files for distribution -->
+						<Strip BaseDir="$(RootDir)" OutputDir="$(SavedOutput)" Platform="Win64" Files="#UnrealGame HoloLens Unstripped" Tag="#UnrealGame HoloLens Stripped"/>
+					</Default>
+				</Switch>
+
+				<Do If="'$(SignExecutables)' == true">
+					<Tag Files="#UnrealGame HoloLens" Filter="$(WindowsSignFilter)" With="#UnrealGame HoloLens Unsigned"/>
+					<Copy Files="#UnrealGame HoloLens Unsigned" From="$(RootDir)" To="$(SavedOutput)" Tag="#Game_ToSign_HoloLens"/>
+					<Sign Files="#Game_ToSign_HoloLens" Tag="#UnrealGame HoloLens Signed"/>
+				</Do>
 			</Node>
 		</Agent>
 		<Aggregate Name="TargetPlatforms_HoloLens" Label="Platforms/HoloLens" Requires="Compile UnrealGame HoloLens"/>
@@ -45,7 +62,7 @@
 
 	<Extend Name="Platform_AddInstalledRequirements">
 		<Do If="$(WithHoloLens)">
-			<Property Name="InstalledRequirements" Value="$(InstalledRequirements);#UnrealGame HoloLens"/>
+			<Property Name="InstalledRequirements" Value="$(InstalledRequirements);#UnrealGame HoloLens;#UnrealGame HoloLens Unstripped;#UnrealGame HoloLens Stripped;#UnrealGame HoloLens Unsigned;#UnrealGame HoloLens Signed"/>
 		</Do>
 	</Extend>
 
@@ -73,6 +90,15 @@
 		</Do>
 	</Extend>
 
+	<Extend Name="Platform_ExtraSignFiles">
+		<Do If="$(WithHoloLens)">
+			<Property Name="ExtraSignFilesHoloLens">
+				Engine\Platforms\Hololens\Binaries\DotNET\AutomationTool\AutomationScripts\...
+			</Property>
+			<Property Name="ExtraSignFiles" Value="$(ExtraSignFiles);$(ExtraSignFilesHoloLens)"/>
+		</Do>
+	</Extend>
+
 	<Extend Name="Platform_Tag_UnrealGame">
 		<Do If="$(WithHoloLens)">
 			<Property Name="CopyHoloLensFilter">
@@ -88,12 +114,11 @@
 
 			<Property Name="CopyHoloLensExceptions">
 			</Property>
-			
-			<Tag Files="#UnrealGame HoloLens;#UnrealGame HoloLens Includes" With="#Installed Win64"/>
+			<Tag Files="#UnrealGame HoloLens" Except="#UnrealGame HoloLens Unstripped;#UnrealGame HoloLens Unsigned" With="#Installed Win64"/>
+			<Tag Files="#UnrealGame HoloLens Stripped;#UnrealGame HoloLens Signed" With="#Saved Output"/>
 			<Property Name="Platforms" Value="$(Platforms)HoloLens;"/>
 			<Property Name="CopyInstalledFilter" Value="$(CopyInstalledFilter);$(CopyHoloLensFilter)"/>
 			<Property Name="CopyInstalledExceptions" Value="$(CopyInstalledExceptions);$(CopyHoloLensExceptions)"/>
-
 		</Do>
 		<Do If="!$(WithHoloLens)">
 			<Property Name="CopyHoloLensCsToolsExceptions">
