Stage OidcToken for installed builds to resolve UE-173034

diff --git a/Engine/Build/InstalledEngineBuild.xml b/Engine/Build/InstalledEngineBuild.xml
--- a/Engine/Build/InstalledEngineBuild.xml
+++ b/Engine/Build/InstalledEngineBuild.xml
@@ -539,6 +539,9 @@
 			<!-- Tag the native P4 libraries as they are not a C# project so not a reference project -->
 			<Tag Files="$(CsToolsDir)/Engine/Binaries/DotNET/EpicGames.Perforce.Native/..." With="#Build Tools CS Binaries" />
 
+			<!-- Tag OidcToken which is already built for all platforms and just needs to be copied  -->
+			<Tag Files="$(CsToolsDir)/Engine/Binaries/DotNET/OidcToken/..." With="#Build Tools CS Binaries" />
+
 			<!-- Create tag set for saved source needed for documentation -->
 			<Tag Files="$(CsToolsDir)/Engine/..." With="#SourceForDocumentation" />
 
