Ensure more files are digitally signed

#rnx

diff --git a/Engine/Build/InstalledEngineBuild.xml b/Engine/Build/InstalledEngineBuild.xml
--- a/Engine/Build/InstalledEngineBuild.xml
+++ b/Engine/Build/InstalledEngineBuild.xml
@@ -990,6 +990,7 @@
 				<Property Name="Platforms" Value="$(Platforms)IOS;"/>
 				<Property Name="CopyInstalledFilter" Value="$(CopyInstalledFilter);$(CopyIOSFilterWin64)"/>
 				<Property Name="CopyInstalledExceptions" Value="$(CopyInstalledExceptions);$(CopyIOSExceptions)"/>
+				<Property Name="ExtraSignFiles" Value="$(ExtraSignFiles);$(ExtraSignFilesIOS)"/>
 			</Do>
 			<Do If="'$(WithTVOS)' == true">
 				<Tag Files="#UnrealGame TVOS" Except="#UnrealGame TVOS Unstripped" With="#Installed Win64"/>
@@ -997,6 +998,7 @@
 				<Property Name="Platforms" Value="$(Platforms)TVOS;"/>
 				<Property Name="CopyInstalledFilter" Value="$(CopyInstalledFilter);$(CopyTVOSFilterWin64)"/>
 				<Property Name="CopyInstalledExceptions" Value="$(CopyInstalledExceptions);$(CopyTVOSExceptions)"/>
+				<Property Name="ExtraSignFiles" Value="$(ExtraSignFiles);$(ExtraSignFilesTVOS)"/>
 			</Do>
 			<Do If="'$(WithLinux)' == true">
 				<Tag Files="#UnrealGame Linux;#Build Tools Linux" Except="#UnrealGame Linux Unstripped" With="#Installed Win64"/>
@@ -1026,6 +1028,7 @@
 			<Tag FileLists="#DependencyLists" With="#Installed Win64" />
 
 			<!-- Confidential Exceptions - done last to ensure they never get overwritten -->
+			<Property Name="SignCopyExceptions" Value="$(CopyInstalledExceptions);$(ConfidentialExceptions)"/>
 			<Property Name="CopyInstalledExceptions" Value="$(CopyInstalledExceptions);#DependencyLists;$(ConfidentialExceptions)"/>
 			
 			<!-- Tag additional files needed for installed build -->
@@ -1036,10 +1039,12 @@
 
 			<Do If="'$(SignExecutables)' == true">
 				<!-- Tag files that need to be signed -->
-				<Tag Files="#Compile UnrealHeaderTool Win64;#Build Tools Win64;#Rules Assemblies" Filter="$(WindowsSignFilter)" Except="$(CopyInstalledExceptions);$(SignExceptions)" With="#Sign Tools for Installed Win64"/>
-				<Tag Files="$(ExtraSignFiles)" Filter="$(WindowsSignFilter)" Except="$(CopyInstalledExceptions);$(SignExceptions)" With="#Sign Tools for Installed Win64"/>
-				<Tag Files="#DatasmithPluginBinaries Win64" Filter="$(WindowsSignFilter)" Except="$(CopyInstalledExceptions);$(SignExceptions)" With="#Sign Tools for Installed Win64" If="$(CompileDatasmithPlugins)"/>
-				<Tag Files="#Build Tools CS Binaries;" Filter="$(WindowsSignFilter)" Except="$(CopyWin64CsToolsExceptions);$(SignExceptions);" With="#Sign CS Tools for Installed Win64"/>
+				<Tag Files="#Compile UnrealHeaderTool Win64;#Build Tools Win64;#Rules Assemblies" Filter="$(WindowsSignFilter)" Except="$(SignCopyExceptions);$(SignExceptions)" With="#Sign Tools for Installed Win64"/>
+				<Tag Files="$(CopyInstalledFilter)" Filter="$(WindowsSignFilter)" Except="$(SignCopyExceptions)" With="#Sign Tools for Installed Win64"/>
+				<Tag FileLists="#DependencyLists" Filter="$(WindowsSignFilter)" Except="$(SignCopyExceptions)" With="#Sign Tools for Installed Win64"/>
+				<Tag Files="$(ExtraSignFiles)" Filter="$(WindowsSignFilter)" Except="$(SignCopyExceptions)" With="#Sign Tools for Installed Win64"/>
+				<Tag Files="#DatasmithPluginBinaries Win64" Filter="$(WindowsSignFilter)" Except="$(SignCopyExceptions)" With="#Sign Tools for Installed Win64" If="$(CompileDatasmithPlugins)"/>
+				<Tag Files="#Build Tools CS Binaries;" Filter="$(WindowsSignFilter)" Except="$(CopyWin64CsToolsExceptions);" With="#Sign CS Tools for Installed Win64"/>
 			</Do>
 
 			<!-- Tag files that need to be copied directly -->
diff --git a/Engine/Build/InstalledEngineFilters.xml b/Engine/Build/InstalledEngineFilters.xml
--- a/Engine/Build/InstalledEngineFilters.xml
+++ b/Engine/Build/InstalledEngineFilters.xml
@@ -251,9 +251,12 @@
 	<Property Name="ExtraSignFilesWin64">
 		<!-- CsvTools binaries which are manually updated in P4 -->
 		Engine/Binaries/DotNET/CsvTools/...
+		Engine/Binaries/ThirdParty/DotNet/6.0.302/windows/...
+		Engine/Binaries/ThirdParty/Python3/...
 		Engine/Binaries/ThirdParty/svn/Win64/...
 		Engine/Binaries/Win64/UnrealTraceServer.exe
 		Engine/Binaries/Win64/EpicWebHelper.exe
+		Engine/Binaries/Win64/embree.dll
 		Engine/Build/Windows/cl-filter/...
 		Engine/Extras/ThirdPartyNotUE/cwrsync/bin/...
 	</Property>
@@ -517,6 +520,13 @@
 	</Property>
 
 	<!-- IOS -->
+	<Property Name="ExtraSignFilesIOS">
+		Engine/Binaries/ThirdParty/DotNet/6.0.302/mac-arm64/...
+		Engine/Binaries/ThirdParty/DotNet/6.0.302/mac-x64/...
+		Engine/Binaries/ThirdParty/Mono/Mac/...
+		Engine/Extras/ThirdPartyNotUE/libimobiledevice/x64/...
+	</Property>
+
 	<Property Name="CopyIOSFilter" Value="$(MacBuildFiles)">
 		<!-- Build Files -->
 		Engine/Build/IOS/...
@@ -549,6 +559,13 @@
 	</Property>
 
 	<!-- TVOS -->
+	<Property Name="ExtraSignFilesTVOS">
+		Engine/Binaries/ThirdParty/DotNet/6.0.302/mac-arm64/...
+		Engine/Binaries/ThirdParty/DotNet/6.0.302/mac-x64/...
+		Engine/Binaries/ThirdParty/Mono/Mac/...
+		Engine/Extras/ThirdPartyNotUE/libimobiledevice/x64/...
+	</Property>
+
 	<Property Name="CopyTVOSFilter" Value="$(MacBuildFiles)">
 		<!-- Build Files -->
 		Engine/Build/TVOS/...
@@ -571,8 +588,10 @@
 	<!-- Linux -->
 	<Property Name="ExtraSignFilesLinux">
 		<!-- Linux cross build Debug tools -->
+		Engine/Binaries/Linux/AgentInterface.dll
 		Engine/Binaries/Linux/dump_syms.exe
 		Engine/Binaries/Linux/BreakpadSymbolEncoder.exe
+		Engine/Binaries/ThirdParty/Mono/Linux/...
 	</Property>
 
 	<Property Name="CopyLinuxFilter">
