- Removed outdated armv7 from android architecture list

diff --git a/Engine/Build/InstalledEngineBuild.xml b/Engine/Build/InstalledEngineBuild.xml
--- a/Engine/Build/InstalledEngineBuild.xml
+++ b/Engine/Build/InstalledEngineBuild.xml
@@ -88,10 +88,8 @@
 	<Property Name="ToolchainDir" Value="$(RootDir)/Engine/Extras/ThirdPartyNotUE/SDKs"/>
 
 	<!-- Architectures that we build for Android. NOTE: This should match the architectures specified in BuildPluginCommand. -->
-	<Property Name="AndroidArchitectures" Value="-armv7+-arm64"/>
+	<Property Name="AndroidArchitectures" Value="arm64"/>
 	
-	<Property Name="AndroidGPUArchitectures" Value="es2"/>
-
 	<Property Name="DefaultArchitectureCommandsWin64" Value=""/>
 
 
