Installed builds - Make sure to copy OidcToken for Linux and Mac

diff --git a/Engine/Build/InstalledEngineBuild.xml b/Engine/Build/InstalledEngineBuild.xml
--- a/Engine/Build/InstalledEngineBuild.xml
+++ b/Engine/Build/InstalledEngineBuild.xml
@@ -595,7 +595,10 @@
 						
 			<!-- Tag the native P4 libraries as they are not a C# project so not a reference project -->
 			<Tag Files="$(CsToolsDir)/Engine/Binaries/DotNET/EpicGames.Perforce.Native/..." With="#Build Tools CS Binaries Mac" />
-						
+			
+			<!-- Tag OidcToken which is already built for all platforms and just needs to be copied  -->
+			<Tag Files="$(CsToolsDir)/Engine/Binaries/DotNET/OidcToken/..." With="#Build Tools CS Binaries Mac" />
+
 			<!-- IPP doesn't currently compile with the bundled Mono. -->
 			<!-- CsCompile Project="$(CsToolsDir)/Engine/Source/Programs/UnrealControls/UnrealControls.csproj" Configuration="Development" Platform="AnyCPU" Tag="#Build Tools CS Binaries Mac;#Sign Binaries Mac" TagReferences="#Build Tools CS Binaries Mac"/-->
 			<!-- CsCompile Project="$(CsToolsDir)/Engine/Source/Programs/IOS/iPhonePackager/iPhonePackager.csproj" Configuration="Development" Platform="AnyCPU" Arguments="/verbosity:minimal /target:Rebuild" Tag="#Build Tools CS Binaries Mac;#Sign Binaries Mac" TagReferences="#Build Tools CS Binaries Mac"/-->
@@ -675,6 +678,9 @@
 			<!-- Tag the native P4 libraries as they are not a C# project so not a reference project -->
 			<Tag Files="$(CsToolsDir)/Engine/Binaries/DotNET/EpicGames.Perforce.Native/..." With="#Build Tools CS Binaries LinuxNative" />
 
+			<!-- Tag OidcToken which is already built for all platforms and just needs to be copied  -->
+			<Tag Files="$(CsToolsDir)/Engine/Binaries/DotNET/OidcToken/..." With="#Build Tools CS Binaries LinuxNative" />
+
 			<!-- Tag AutomationTool Script module build records, so that prebuilt modules may be discovered in the absence of source code -->
 			<Tag Files="$(CsToolsDir)/Engine/Intermediate/ScriptModules/..." With="#Build Tools CS Binaries LinuxNative" />
 
