Updating installed engine build scripts to no longer compile UnrealHeaderTool.

diff --git a/Engine/Build/InstalledEngineBuild.xml b/Engine/Build/InstalledEngineBuild.xml
--- a/Engine/Build/InstalledEngineBuild.xml
+++ b/Engine/Build/InstalledEngineBuild.xml
@@ -150,10 +150,7 @@
 		<Node Name="Update Version Files">
 			<SetVersion Change="$(Change)" Branch="$(EscapedBranch)" If="$(IsBuildMachine)"/>
 		</Node>
-		<Node Name="Compile UnrealHeaderTool Win64" Requires="Update Version Files">
-			<Compile Target="UnrealHeaderTool" Platform="Win64" Configuration="Development" Arguments="-precompile -allmodules"/>
-		</Node>
-		<Node Name="Compile UnrealEditor Win64" Requires="Compile UnrealHeaderTool Win64" Produces="#UnrealEditor Win64;#UnrealEditor Win64 Unstripped;#UnrealEditor Win64 Stripped;#UnrealEditor Win64 Unsigned;#UnrealEditor Win64 Signed;#UnrealEditor Win64 WithExceptions">
+		<Node Name="Compile UnrealEditor Win64" Requires="Update Version Files" Produces="#UnrealEditor Win64;#UnrealEditor Win64 Unstripped;#UnrealEditor Win64 Stripped;#UnrealEditor Win64 Unsigned;#UnrealEditor Win64 Signed;#UnrealEditor Win64 WithExceptions">
 			<Compile Target="UnrealEditor" Platform="Win64" Configuration="DebugGame" Tag="#UnrealEditor Win64" Arguments="-precompile -allmodules $(CrashReporterCompileArgs)"/>
 			<Compile Target="UnrealEditor" Platform="Win64" Configuration="Development" Tag="#UnrealEditor Win64" Arguments="-precompile -allmodules $(CrashReporterCompileArgs)"/>
 
@@ -192,10 +189,7 @@
 
 	<!-- Base editor and tools for compile on Mac -->
 	<Agent Name="Editor Mac" Type="Mac_Licensee">
-		<Node Name="Compile UnrealHeaderTool Mac" Requires="Update Version Files">
-			<Compile Target="UnrealHeaderTool" Platform="Mac" Configuration="Development" Arguments="-precompile -allmodules"/>
-		</Node>
-		<Node Name="Compile UnrealEditor Mac" Requires="Compile UnrealHeaderTool Mac" Produces="#UnrealEditor Mac;#UnrealEditor Mac Unstripped;#UnrealEditor Mac Stripped;#UnrealEditor Mac Unsigned;#UnrealEditor Mac Signed;#UnrealEditor Mac WithExceptions">
+		<Node Name="Compile UnrealEditor Mac" Requires="Update Version Files" Produces="#UnrealEditor Mac;#UnrealEditor Mac Unstripped;#UnrealEditor Mac Stripped;#UnrealEditor Mac Unsigned;#UnrealEditor Mac Signed;#UnrealEditor Mac WithExceptions">
 			<Property Name="AllModulesArg" Value="-allmodules"/>
 			<Property Name="AllModulesArg" Value="" If="$(WithMacArm64)"/>
 			<Compile Target="UnrealEditor" Platform="Mac" Configuration="DebugGame" Tag="#UnrealEditor Mac" Arguments="-precompile -EnableDSYM $(AllModulesArg) $(CrashReporterCompileArgs) $(UseArm64ArchitectureArg)"/>
@@ -229,10 +223,7 @@
 
 	<!-- Base editor and tools for compile on Linux -->
 	<Agent Name="Editor Linux" Type="Linux_Licensee">
-		<Node Name="Compile UnrealHeaderTool Linux" Requires="Update Version Files">
-			<Compile Target="UnrealHeaderTool" Platform="Linux" Configuration="Development" Arguments="-precompile -allmodules"/>
-		</Node>
-		<Node Name="Compile UnrealEditor Linux" Requires="Compile UnrealHeaderTool Linux" Produces="#UnrealEditor Linux;#UnrealEditor Linux Unstripped;#UnrealEditor Linux Stripped;#UnrealEditor Linux WithExceptions">
+		<Node Name="Compile UnrealEditor Linux" Requires="Update Version Files" Produces="#UnrealEditor Linux;#UnrealEditor Linux Unstripped;#UnrealEditor Linux Stripped;#UnrealEditor Linux WithExceptions">
 			<Compile Target="UnrealEditor" Platform="Linux" Configuration="DebugGame" Tag="#UnrealEditor Linux" Arguments="-precompile -allmodules $(CrashReporterCompileArgs)"/>
 			<Compile Target="UnrealEditor" Platform="Linux" Configuration="Development" Tag="#UnrealEditor Linux" Arguments="-precompile -allmodules $(CrashReporterCompileArgs)"/>
 
@@ -257,7 +248,7 @@
 	<Agent Name="Target Win64" Type="Win64_Licensee">
 		<Property Name="OptionalTargetTags" Value=""/>
 		<ForEach Name="Target" Values="$(OptionalClientTarget);$(OptionalServerTarget)">
-			<Node Name="Compile $(Target) Win64" Requires="Compile UnrealHeaderTool Win64" Produces="#$(Target) Win64">
+			<Node Name="Compile $(Target) Win64" Requires="Update Version Files" Produces="#$(Target) Win64">
 				<ForEach Name="Configuration" Values="$(GameConfigurations)">
 					<Compile Target="$(Target)" Platform="Win64" Configuration="$(Configuration)" Tag="#$(Target) Win64" Arguments="-precompile -allmodules -nolink $(TargetDebugInfoArg)"/>
 					<Compile Target="$(Target)" Platform="Win64" Configuration="$(Configuration)" Tag="#$(Target) Win64" Arguments="-precompile $(TargetDebugInfoArg)" Clean="false"/>
@@ -265,7 +256,7 @@
 				</ForEach>
 			</Node>
 		</ForEach>
-		<Node Name="Compile UnrealGame Win64" Requires="Compile UnrealHeaderTool Win64;$(OptionalTargetTags)" Produces="#UnrealGame Win64;#UnrealGame Win64 Unstripped;#UnrealGame Win64 Stripped;#UnrealGame Win64 Unsigned;#UnrealGame Win64 Signed">
+		<Node Name="Compile UnrealGame Win64" Requires="Update Version Files;$(OptionalTargetTags)" Produces="#UnrealGame Win64;#UnrealGame Win64 Unstripped;#UnrealGame Win64 Stripped;#UnrealGame Win64 Unsigned;#UnrealGame Win64 Signed">
 			<ForEach Name="Tag" Values="$(OptionalTargetTags)" If="'$(OptionalTargetTags)' != ''">
 				<Tag Files="$(Tag)" With="#UnrealGame Win64"/>
 			</ForEach>
@@ -304,7 +295,7 @@
 	<Aggregate Name="TargetPlatforms_Win64" Requires="Compile UnrealGame Win64"/>
 
 	<Agent Name="Target Linux" Type="Win64_Licensee">
-		<Node Name="Compile UnrealGame Linux" Requires="Compile UnrealHeaderTool Win64" Produces="#UnrealGame Linux;#UnrealGame Linux Unstripped;#UnrealGame Linux Stripped">
+		<Node Name="Compile UnrealGame Linux" Requires="Update Version Files" Produces="#UnrealGame Linux;#UnrealGame Linux Unstripped;#UnrealGame Linux Stripped">
 			<ForEach Name="Target" Values="UnrealGame;$(OptionalClientTarget);$(OptionalServerTarget)">
 				<ForEach Name="Configuration" Values="$(GameConfigurations)">
 					<Compile Target="$(Target)" Platform="Linux" Configuration="$(Configuration)" Tag="#UnrealGame Linux" AllowParallelExecutor="$(bAllowParallelExecutor)" Arguments="-precompile -allmodules -nolink $(TargetDebugInfoArg)"/>
@@ -319,7 +310,7 @@
 	<Aggregate Name="TargetPlatforms_Linux" Requires="Compile UnrealGame Linux"/>
 
 	<Agent Name="Target LinuxNative" Type="Linux_Licensee">
-		<Node Name="Compile UnrealGame LinuxNative" Requires="Compile UnrealHeaderTool Linux" Produces="#UnrealGame LinuxNative;#UnrealGame LinuxNative Unstripped;#UnrealGame LinuxNative Stripped">
+		<Node Name="Compile UnrealGame LinuxNative" Requires="Update Version Files" Produces="#UnrealGame LinuxNative;#UnrealGame LinuxNative Unstripped;#UnrealGame LinuxNative Stripped">
 			<ForEach Name="Target" Values="UnrealGame;$(OptionalClientTarget);$(OptionalServerTarget)">
 				<ForEach Name="Configuration" Values="$(GameConfigurations)">
 					<Compile Target="$(Target)" Platform="Linux" Configuration="$(Configuration)" Tag="#UnrealGame LinuxNative" AllowParallelExecutor="$(bAllowParallelExecutor)" Arguments="-precompile -allmodules -nolink $(TargetDebugInfoArg)"/>
@@ -334,7 +325,7 @@
 	<Aggregate Name="TargetPlatforms_LinuxNative" Requires="Compile UnrealGame LinuxNative"/>
 
 	<Agent Name="Target LinuxArm64" Type="Win64_Licensee">
-		<Node Name="Compile UnrealGame LinuxArm64" Requires="Compile UnrealHeaderTool Win64" Produces="#UnrealGame LinuxArm64;#UnrealGame LinuxArm64 Unstripped;#UnrealGame LinuxArm64 Stripped">
+		<Node Name="Compile UnrealGame LinuxArm64" Requires="Update Version Files" Produces="#UnrealGame LinuxArm64;#UnrealGame LinuxArm64 Unstripped;#UnrealGame LinuxArm64 Stripped">
 			<ForEach Name="Target" Values="UnrealGame;$(OptionalClientTarget);$(OptionalServerTarget)">
 				<ForEach Name="Configuration" Values="$(GameConfigurations)">
 					<Compile Target="$(Target)" Platform="LinuxArm64" Configuration="$(Configuration)" Tag="#UnrealGame LinuxArm64" AllowParallelExecutor="$(bAllowParallelExecutor)" Arguments="-precompile -allmodules -nolink $(TargetDebugInfoArg)"/>
@@ -349,7 +340,7 @@
 	<Aggregate Name="TargetPlatforms_LinuxArm64" Label="Platforms/Linux (AArch64)" Requires="Compile UnrealGame LinuxArm64"/>
 
 	<Agent Name="Target LinuxArm64 LinuxNative" Type="Linux_Licensee">
-		<Node Name="Compile UnrealGame LinuxArm64 LinuxNative" Requires="Compile UnrealHeaderTool Linux" Produces="#UnrealGame LinuxArm64 LinuxNative;#UnrealGame LinuxArm64 LinuxNative Unstripped;#UnrealGame LinuxArm64 LinuxNative Stripped">
+		<Node Name="Compile UnrealGame LinuxArm64 LinuxNative" Requires="Update Version Files" Produces="#UnrealGame LinuxArm64 LinuxNative;#UnrealGame LinuxArm64 LinuxNative Unstripped;#UnrealGame LinuxArm64 LinuxNative Stripped">
 			<ForEach Name="Target" Values="UnrealGame;$(OptionalClientTarget);$(OptionalServerTarget)">
 				<ForEach Name="Configuration" Values="$(GameConfigurations)">
 					<Compile Target="$(Target)" Platform="LinuxArm64" Configuration="$(Configuration)" Tag="#UnrealGame LinuxArm64 LinuxNative" AllowParallelExecutor="$(bAllowParallelExecutor)" Arguments="-precompile -allmodules -nolink $(TargetDebugInfoArg)"/>
@@ -364,7 +355,7 @@
 	<Aggregate Name="TargetPlatforms_LinuxArm64LinuxNative" Label="Platforms/Linux (AArch64)" Requires="Compile UnrealGame LinuxArm64 LinuxNative"/>
 
 	<Agent Name="Target Android" Type="Win64_Licensee">
-		<Node Name="Compile UnrealGame Android" Requires="Compile UnrealHeaderTool $(HostPlatform)" Produces="#UnrealGame Android;#UnrealGame Android Unstripped;#UnrealGame Android Stripped">
+		<Node Name="Compile UnrealGame Android" Requires="Update Version Files" Produces="#UnrealGame Android;#UnrealGame Android Unstripped;#UnrealGame Android Stripped">
 			<ForEach Name="Target" Values="UnrealGame;$(OptionalClientTarget)">
 				<ForEach Name="Configuration" Values="$(GameConfigurations)">
 					<Compile Target="$(Target)" Platform="Android" Configuration="$(Configuration)" Tag="#UnrealGame Android" AllowParallelExecutor="$(bAllowParallelExecutor)" Arguments="-precompile -allmodules -nolink -architectures=$(AndroidArchitectures) $(TargetDebugInfoArg)"/>
@@ -384,7 +375,7 @@
 	<Agent Name="Target Mac" Type="Mac_Licensee">
 		<Property Name="OptionalTargetTags" Value=""/>
 		<ForEach Name="Target" Values="$(OptionalClientTarget);$(OptionalServerTarget)">
-			<Node Name="Compile $(Target) Mac" Requires="Compile UnrealHeaderTool Mac" Produces="#$(Target) Mac">
+			<Node Name="Compile $(Target) Mac" Requires="Update Version Files" Produces="#$(Target) Mac">
 				<ForEach Name="Configuration" Values="$(GameConfigurations)">
 					<Compile Target="$(Target)" Platform="Mac" Configuration="$(Configuration)" Tag="#$(Target) Mac" Arguments="-precompile -allmodules -nolink $(TargetDebugInfoArgMac)"/>
 					<Compile Target="$(Target)" Platform="Mac" Configuration="$(Configuration)" Tag="#$(Target) Mac" Arguments="-precompile $(TargetDebugInfoArgMac)" Clean="false"/>
@@ -393,7 +384,7 @@
 			</Node>
 		</ForEach>
 
-		<Node Name="Compile UnrealGame Mac" Requires="Compile UnrealHeaderTool Mac;$(OptionalTargetTags)" Produces="#UnrealGame Mac;#UnrealGame Mac Unstripped;#UnrealGame Mac Stripped;#UnrealGame Mac Unsigned;#UnrealGame Mac Signed">
+		<Node Name="Compile UnrealGame Mac" Requires="Update Version Files;$(OptionalTargetTags)" Produces="#UnrealGame Mac;#UnrealGame Mac Unstripped;#UnrealGame Mac Stripped;#UnrealGame Mac Unsigned;#UnrealGame Mac Signed">
 			<ForEach Name="Tag" Values="$(OptionalTargetTags)" If="'$(OptionalTargetTags)' != ''">
 				<Tag Files="$(Tag)" With="#UnrealGame Mac"/>
 			</ForEach>
@@ -417,7 +408,7 @@
 	<Aggregate Name="TargetPlatforms_Mac" Requires="Compile UnrealGame Mac"/>
 
 	<Agent Name="Target IOS" Type="Mac_Licensee">
-		<Node Name="Compile UnrealGame IOS" Requires="Compile UnrealHeaderTool Mac" Produces="#UnrealGame IOS;#UnrealGame IOS Unstripped;#UnrealGame IOS Stripped">
+		<Node Name="Compile UnrealGame IOS" Requires="Update Version Files" Produces="#UnrealGame IOS;#UnrealGame IOS Unstripped;#UnrealGame IOS Stripped">
 			<ForEach Name="Target" Values="UnrealGame;$(OptionalClientTarget)">
 				<ForEach Name="Configuration" Values="$(GameConfigurations)">
 					<Compile Target="$(Target)" Platform="IOS" Configuration="$(Configuration)" Tag="#UnrealGame IOS" AllowParallelExecutor="$(bAllowParallelExecutor)" Arguments="-precompile -allmodules -nolink $(TargetDebugInfoArgMac)"/>
@@ -432,7 +423,7 @@
 	<Aggregate Name="TargetPlatforms_IOS" Label="Platforms/IOS" Requires="Compile UnrealGame IOS"/>
 
 	<Agent Name="Target TVOS" Type="Mac_Licensee">
-		<Node Name="Compile UnrealGame TVOS" Requires="Compile UnrealHeaderTool Mac" Produces="#UnrealGame TVOS;#UnrealGame TVOS Unstripped;#UnrealGame TVOS Stripped">
+		<Node Name="Compile UnrealGame TVOS" Requires="Update Version Files" Produces="#UnrealGame TVOS;#UnrealGame TVOS Unstripped;#UnrealGame TVOS Stripped">
 			<ForEach Name="Target" Values="UnrealGame;$(OptionalClientTarget)">
 				<ForEach Name="Configuration" Values="$(GameConfigurations)">
 					<Compile Target="$(Target)" Platform="TVOS" Configuration="$(Configuration)" Tag="#UnrealGame TVOS" AllowParallelExecutor="$(bAllowParallelExecutor)" Arguments="-precompile -allmodules -nolink $(TargetDebugInfoArgMac)"/>
@@ -447,7 +438,7 @@
 	<Aggregate Name="TargetPlatforms_TVOS" Label="Platforms/TVOS" Requires="Compile UnrealGame TVOS"/>
 
 	<Agent Name="NotForLicensee Tools Group Win64" Type="Win64">
-		<Node Name="Build NotForLicensee Tools Win64" Requires="Compile UnrealHeaderTool Win64" Produces="#NotForLicensee Build Tools Win64">
+		<Node Name="Build NotForLicensee Tools Win64" Requires="Update Version Files" Produces="#NotForLicensee Build Tools Win64">
 			<Compile Target="DatasmithCADWorker" Configuration="Development" Platform="Win64" Tag="#NotForLicensee Build Tools Win64"/>
 			<!-- Temporarily removing the compilation of DatasmithCADRuntime from the installed build
 			<Compile Target="DatasmithCADRuntime" Configuration="Development" Platform="Win64" Tag="#NotForLicensee Build Tools Win64"/>
@@ -466,7 +457,7 @@
 	</Agent>
 
 	<Agent Name="Tools Group Win64" Type="Win64_Licensee">
-		<Node Name="Build Tools Win64" Requires="Compile UnrealHeaderTool Win64;#NotForLicensee Build Tools Win64" Produces="#Tools Win64 Complete">
+		<Node Name="Build Tools Win64" Requires="Update Version Files;#NotForLicensee Build Tools Win64" Produces="#Tools Win64 Complete">
 			<Compile Target="CrashReportClient" Configuration="Shipping" Platform="Win64" Tag="#Build Tools Win64" Arguments="$(CrashReporterCompileArgs)"/>
 			<Compile Target="CrashReportClientEditor" Configuration="Shipping" Platform="Win64" Tag="#Build Tools Win64" Arguments="$(CrashReporterCompileArgs)"/>
 			<Tag Files="#NotForLicensee Build Tools Win64" With="#Build Tools Win64"/>
@@ -546,7 +537,7 @@
 			<!-- Sign the binaries -->
 			<Sign Files="#Sign Binaries" If="'$(SignExecutables)' == true"/>
 		</Node>
-		<Node Name="Build Tools Linux" Requires="Compile UnrealHeaderTool Win64">
+		<Node Name="Build Tools Linux" Requires="Update Version Files">
 			<Compile Target="CrashReportClient" Platform="Linux" Configuration="Shipping" Tag="#Build Tools Linux" Arguments="$(CrashReporterCompileArgs)"/>
 			<Compile Target="CrashReportClientEditor" Platform="Linux" Configuration="Shipping" Tag="#Build Tools Linux" Arguments="$(CrashReporterCompileArgs)"/>
 			<Compile Target="ShaderCompileWorker" Configuration="Development" Platform="Linux" Tag="#Build Tools Linux"/>
@@ -614,7 +605,7 @@
 			<Sign Files="#Sign Binaries Mac" If="'$(SignExecutables)' == true"/>
 		</Node>
 	
-		<Node Name="Build Tools Mac" Requires="Compile UnrealHeaderTool Mac">
+		<Node Name="Build Tools Mac" Requires="Update Version Files">
 			<Compile Target="CrashReportClient" Configuration="Shipping" Platform="Mac" Tag="#Build Tools Mac" Arguments="-EnableDSYM $(CrashReporterCompileArgs)"/>
 			<Compile Target="CrashReportClientEditor" Configuration="Shipping" Platform="Mac" Tag="#Build Tools Mac" Arguments="-EnableDSYM $(CrashReporterCompileArgs)"/>
 			<Compile Target="ShaderCompileWorker" Configuration="Development" Platform="Mac" Tag="#Build Tools Mac" Arguments="-EnableDSYM"/>
@@ -636,7 +627,7 @@
 	</Agent>
 
 	<Agent Name="Tools Group LinuxNative" Type="Linux_Licensee">
-		<Node Name="Build Tools LinuxNative" Requires="Compile UnrealHeaderTool Linux">
+		<Node Name="Build Tools LinuxNative" Requires="Update Version Files">
 			<Compile Target="CrashReportClient" Platform="Linux" Configuration="Shipping" Tag="#Build Tools LinuxNative" Arguments="$(CrashReporterCompileArgs)"/>
 			<Compile Target="CrashReportClientEditor" Platform="Linux" Configuration="Shipping"  Tag="#Build Tools LinuxNative" Arguments="$(CrashReporterCompileArgs)"/>
 			<Compile Target="ShaderCompileWorker" Configuration="Development" Platform="Linux" Tag="#Build Tools LinuxNative"/>
@@ -786,7 +777,7 @@
 			<Delete Files="$(DDCDir)\..."/>
 			<!-- Build up a list of files needed to build DDC -->
 			<CsCompile Project="Engine/Source/Programs/UnrealBuildTool/UnrealBuildTool.csproj" Configuration="Development" Platform="AnyCPU" Tag="#ToCopy" TagReferences="#ToCopy" EnumerateOnly="true"/>
-			<Tag Files="#Compile UnrealHeaderTool Win64;#UnrealEditor Win64;#Build Tools Win64" With="#ToCopy"/>
+			<Tag Files="#UnrealEditor Win64;#Build Tools Win64" With="#ToCopy"/>
 			<Tag Files="#ToCopy" Filter="*.target" With="#TargetReceipts"/>
 			<TagReceipt Files="#TargetReceipts" RuntimeDependencies="true" With="#ToCopy"/>
 			<Tag Files="Engine\Build\Build.version" With="#ToCopy"/>
@@ -831,7 +822,7 @@
 			<Delete Files="$(DDCDir)/..."/>
 			<!-- Build up a list of files needed to build DDC -->
 			<CsCompile Project="Engine/Source/Programs/UnrealBuildTool/UnrealBuildTool.csproj" Configuration="Development" Platform="AnyCPU" Tag="#ToCopy" TagReferences="#ToCopy" EnumerateOnly="true"/>
-			<Tag Files="#Compile UnrealHeaderTool Mac;#UnrealEditor Mac;#Build Tools Mac" With="#ToCopy"/>
+			<Tag Files="#UnrealEditor Mac;#Build Tools Mac" With="#ToCopy"/>
 			<Tag Files="#ToCopy" Filter="*.target" With="#TargetReceipts"/>
 			<TagReceipt Files="#TargetReceipts" RuntimeDependencies="true" With="#ToCopy"/>
 			<Tag Files="Engine/Build/Build.version" With="#ToCopy"/>
@@ -870,7 +861,7 @@
 			<Delete Files="$(DDCDir)/..."/>
 			<!-- Build up a list of files needed to build DDC -->
 			<CsCompile Project="Engine/Source/Programs/UnrealBuildTool/UnrealBuildTool.csproj" Configuration="Development" Platform="AnyCPU" Tag="#ToCopy" TagReferences="#ToCopy" EnumerateOnly="true"/>
-			<Tag Files="#Compile UnrealHeaderTool Linux;#UnrealEditor Linux;#Build Tools LinuxNative" With="#ToCopy"/>
+			<Tag Files="#UnrealEditor Linux;#Build Tools LinuxNative" With="#ToCopy"/>
 			<Tag Files="#ToCopy" Filter="*.target" With="#TargetReceipts"/>
 			<TagReceipt Files="#TargetReceipts" RuntimeDependencies="true" With="#ToCopy"/>
 			<Tag Files="Engine/Content/..." Except="*.psd;*.pdn;*.fbx;*.po" With="#ToCopy"/>
@@ -898,7 +889,7 @@
 
 	<!-- These are plugins that need to be compiled with a full workspace rather than $(Platform)_Licensee. -->
 	<Agent Name="Extra Plugins Agent Win64" Type="Win64">
-		<Node Name="Compile Datasmith Plugins Win64" Requires="Compile UnrealHeaderTool Win64" Produces="#DatasmithPluginBinaries Win64">
+		<Node Name="Compile Datasmith Plugins Win64" Requires="Update Version Files" Produces="#DatasmithPluginBinaries Win64">
 			<Do If="$(WithWin64) And $(CompileDatasmithPlugins) == 'True'">
 				<Property Name="DatasmithPlugins" Value="DatasmithCADImporter;DatasmithIFCImporter;DatasmithC4DImporter;AxFImporter;MDLImporter"/>
 				<ForEach Name="Plugin" Values="$(DatasmithPlugins)">
@@ -910,7 +901,7 @@
 	</Agent>
 
 	<Agent Name="Extra Plugins Agent Mac" Type="Mac">
-		<Node Name="Compile Datasmith Plugins Mac" Requires="Compile UnrealHeaderTool Mac" Produces="#DatasmithPluginBinaries Mac">
+		<Node Name="Compile Datasmith Plugins Mac" Requires="Update Version Files" Produces="#DatasmithPluginBinaries Mac">
 			<Do If="$(WithMac) And $(CompileDatasmithPlugins) == 'True'">
 				<Property Name="DatasmithPlugins" Value="DatasmithCADImporter;DatasmithIFCImporter;DatasmithC4DImporter;AxFImporter;MDLImporter"/>
 				<ForEach Name="Plugin" Values="$(DatasmithPlugins)">
@@ -922,7 +913,7 @@
 	</Agent>
 
 	<Agent Name="Extra Plugins Agent Linux" Type="Linux">
-		<Node Name="Compile Datasmith Plugins Linux" Requires="Compile UnrealHeaderTool Linux" Produces="#DatasmithPluginBinaries Linux">
+		<Node Name="Compile Datasmith Plugins Linux" Requires="Update Version Files" Produces="#DatasmithPluginBinaries Linux">
 			<Do If="$(WithLinux) And $(CompileDatasmithPlugins) == 'True'">
 				<Property Name="DatasmithPlugins" Value="DatasmithCADImporter;MDLImporter"/>
 				<ForEach Name="Plugin" Values="$(DatasmithPlugins)">
@@ -937,12 +928,12 @@
 		<!-- Build up a list of requirements as some can be optional -->
 		<Property Name="UnrealEditor" Value="#UnrealEditor Win64" />
 		<Property Name="UnrealEditor" Value="#UnrealEditor Win64 WithExceptions" If="$(CompileDatasmithPlugins)" />
-		<Property Name="InstalledRequirements" Value="Compile UnrealHeaderTool Win64;$(UnrealEditor);#UnrealEditor Win64 Unstripped;#UnrealEditor Win64 Stripped;#UnrealEditor Win64 Unsigned;#UnrealEditor Win64 Signed;#Build Tools Win64;#Build Tools CS Binaries;#Rules Assemblies"/>
+		<Property Name="InstalledRequirements" Value="Update Version Files;$(UnrealEditor);#UnrealEditor Win64 Unstripped;#UnrealEditor Win64 Stripped;#UnrealEditor Win64 Unsigned;#UnrealEditor Win64 Signed;#Build Tools Win64;#Build Tools CS Binaries;#Rules Assemblies"/>
 		<!-- Optional Target Platform requirements -->
 		<Expand Name="Platform_AddInstalledRequirements"/>
 		<Property Name="InstalledRequirements" Value="$(InstalledRequirements);#UnrealGame Win64;#UnrealGame Win64 Unstripped;#UnrealGame Win64 Stripped;#UnrealGame Win64 Unsigned;#UnrealGame Win64 Signed" If="'$(WithWin64)' == true"/>
 		<Property Name="InstalledRequirements" Value="$(InstalledRequirements);#UnrealGame Android;#UnrealGame Android Unstripped;#UnrealGame Android Stripped" If="'$(WithAndroid)' == true"/>
-		<Property Name="InstalledRequirements" Value="$(InstalledRequirements);#UnrealGame IOS;#UnrealGame IOS Unstripped;#UnrealGame IOS Stripped;#Compile UnrealHeaderTool Mac" If="'$(WithIOS)' == true"/>
+		<Property Name="InstalledRequirements" Value="$(InstalledRequirements);#UnrealGame IOS;#UnrealGame IOS Unstripped;#UnrealGame IOS Stripped;#Update Version Files" If="'$(WithIOS)' == true"/>
 		<Property Name="InstalledRequirements" Value="$(InstalledRequirements);#UnrealGame TVOS;#UnrealGame TVOS Unstripped;#UnrealGame TVOS Stripped" If="'$(WithTVOS)' == true"/>
 		<Property Name="InstalledRequirements" Value="$(InstalledRequirements);#UnrealGame Linux;#UnrealGame Linux Unstripped;#UnrealGame Linux Stripped;#Build Tools Linux" If="'$(WithLinux)' == true"/>
 		<Property Name="InstalledRequirements" Value="$(InstalledRequirements);#UnrealGame LinuxArm64;#UnrealGame LinuxArm64 Unstripped;#UnrealGame LinuxArm64 Stripped" If="'$(WithLinuxArm64)' == true"/>
@@ -966,7 +957,6 @@
 			<Property Name="ExtraSignFiles" Value="$(ExtraSignFilesWin64)"/>
 
 			<!-- Tag files that always need to be part of an installed build -->
-			<Tag Files="#Compile UnrealHeaderTool Win64" With="#Installed Win64"/>
 			<Tag Files="$(UnrealEditor)" Except="#UnrealEditor Win64 Unstripped;#UnrealEditor Win64 Unsigned" With="#Installed Win64"/>
 			<Tag Files="#UnrealEditor Win64 Stripped;#UnrealEditor Win64 Signed" With="#Saved Output"/>
 			<Tag Files="#Build Tools Win64;#Make Feature Packs;#Rules Assemblies" With="#Installed Win64"/>
@@ -989,7 +979,7 @@
 				<Property Name="ExtraSignFiles" Value="$(ExtraSignFiles);$(ExtraSignFilesAndroid)"/>
 			</Do>
 			<Do If="'$(WithIOS)' == true">
-				<Tag Files="#UnrealGame IOS;#Compile UnrealHeaderTool Mac" Except="#UnrealGame IOS Unstripped" With="#Installed Win64"/>
+				<Tag Files="#UnrealGame IOS" Except="#UnrealGame IOS Unstripped" With="#Installed Win64"/>
 				<Tag Files="#UnrealGame IOS Stripped" With="#Saved Output"/>
 				<Property Name="Platforms" Value="$(Platforms)IOS;"/>
 				<Property Name="CopyInstalledFilter" Value="$(CopyInstalledFilter);$(CopyIOSFilterWin64)"/>
@@ -1039,11 +1029,11 @@
 			<Tag Files="$(CopyInstalledFilter)" With="#Installed Win64"/>
 
 			<!-- Tag Win64 files that need to be stripped -->
-			<Tag Files="#Compile UnrealHeaderTool Win64;#Build Tools Win64" Filter="$(Win64StripFilter)" Except="$(Win64StripExceptions);$(CopyInstalledExceptions)" With="#Strip Tools for Installed Win64"/>
+			<Tag Files="#Build Tools Win64" Filter="$(Win64StripFilter)" Except="$(Win64StripExceptions);$(CopyInstalledExceptions)" With="#Strip Tools for Installed Win64"/>
 
 			<Do If="'$(SignExecutables)' == true">
 				<!-- Tag files that need to be signed -->
-				<Tag Files="#Compile UnrealHeaderTool Win64;#Build Tools Win64;#Rules Assemblies" Filter="$(WindowsSignFilter)" Except="$(SignCopyExceptions);$(SignExceptions)" With="#Sign Tools for Installed Win64"/>
+				<Tag Files="#Build Tools Win64;#Rules Assemblies" Filter="$(WindowsSignFilter)" Except="$(SignCopyExceptions);$(SignExceptions)" With="#Sign Tools for Installed Win64"/>
 				<Tag Files="$(CopyInstalledFilter)" Filter="$(WindowsSignFilter)" Except="$(SignCopyExceptions)" With="#Sign Tools for Installed Win64"/>
 				<Tag FileLists="#DependencyLists" Filter="$(WindowsSignFilter)" Except="$(SignCopyExceptions)" With="#Sign Tools for Installed Win64"/>
 				<Tag Files="$(ExtraSignFiles)" Filter="$(WindowsSignFilter)" Except="$(SignCopyExceptions)" With="#Sign Tools for Installed Win64"/>
@@ -1097,7 +1087,7 @@
 		<!-- Build up a list of requirements as some can be optional -->
 		<Property Name="UnrealEditor" Value="#UnrealEditor Mac" />
 		<Property Name="UnrealEditor" Value="#UnrealEditor Mac WithExceptions" If="$(CompileDatasmithPlugins)" />
-		<Property Name="FilterRequirements" Value="#Compile UnrealHeaderTool Mac;$(UnrealEditor);#UnrealEditor Mac Unstripped;#UnrealEditor Mac Stripped;#UnrealEditor Mac Unsigned;#UnrealEditor Mac Signed;#Build Tools Mac;#Build Tools CS Binaries Mac;#Rules Assemblies"/>
+		<Property Name="FilterRequirements" Value="#Update Version Files;$(UnrealEditor);#UnrealEditor Mac Unstripped;#UnrealEditor Mac Stripped;#UnrealEditor Mac Unsigned;#UnrealEditor Mac Signed;#Build Tools Mac;#Build Tools CS Binaries Mac;#Rules Assemblies"/>
 
 		<!-- Optional Target Platform requirements -->
 		<Property Name="FilterRequirements" Value="$(FilterRequirements);#UnrealGame Mac;#UnrealGame Mac Unstripped;#UnrealGame Mac Stripped;#UnrealGame Mac Unsigned;#UnrealGame Mac Signed" If="'$(WithMac)' == true"/>
@@ -1123,7 +1113,6 @@
 			<Property Name="CopyInstalledExceptions" Value="$(CopyEditorExceptionsMac)"/>
 
 			<!-- Tag files that always need to be part of an installed build -->
-			<Tag Files="#Compile UnrealHeaderTool Mac" With="#Installed Mac"/>
 			<Tag Files="$(UnrealEditor)" Except="#UnrealEditor Mac Unstripped;#UnrealEditor Mac Unsigned" With="#Installed Mac"/>
 			<Tag Files="#UnrealEditor Mac Stripped;#UnrealEditor Mac Signed" With="#Saved Output"/>
 			<Tag Files="#Build Tools Mac;#Make Feature Packs Mac;#Rules Assemblies" With="#Installed Mac"/>
@@ -1174,11 +1163,11 @@
 			<Tag Files="$(CopyInstalledFilter)" With="#Installed Mac"/>
 
 			<!-- Tag Mac files that need to be stripped -->
-			<Tag Files="#Compile UnrealHeaderTool Mac;#Build Tools Mac" Filter="$(MacStripFilter)" Except="$(CopyInstalledExceptions)" With="#Strip Tools for Installed Mac"/>
+			<Tag Files="#Build Tools Mac" Filter="$(MacStripFilter)" Except="$(CopyInstalledExceptions)" With="#Strip Tools for Installed Mac"/>
 
 			<Do If="'$(SignExecutables)' == true">
 				<!-- Tag files that need to be signed -->
-				<Tag Files="#Compile UnrealHeaderTool Mac;#Build Tools Mac" Filter="$(MacSignFilter)" Except="$(CopyInstalledExceptions);$(SignExceptions)" With="#Sign Tools for Installed Mac"/>
+				<Tag Files="#Build Tools Mac" Filter="$(MacSignFilter)" Except="$(CopyInstalledExceptions);$(SignExceptions)" With="#Sign Tools for Installed Mac"/>
 			</Do>
 
 			<!-- Tag files that need to be copied directly -->
@@ -1222,7 +1211,7 @@
 		<!-- Build up a list of requirements as some can be optional -->
 		<Property Name="UnrealEditor" Value="#UnrealEditor Linux" />
 		<Property Name="UnrealEditor" Value="#UnrealEditor Linux WithExceptions" If="$(CompileDatasmithPlugins)" />
-		<Property Name="InstalledRequirements" Value="#Compile UnrealHeaderTool Linux;$(UnrealEditor);#UnrealEditor Linux Unstripped;#UnrealEditor Linux Stripped;#Build Tools LinuxNative;#Build Tools CS Binaries LinuxNative;#Rules Assemblies"/>
+		<Property Name="InstalledRequirements" Value="#Update Version Files;$(UnrealEditor);#UnrealEditor Linux Unstripped;#UnrealEditor Linux Stripped;#Build Tools LinuxNative;#Build Tools CS Binaries LinuxNative;#Rules Assemblies"/>
 
 		<!-- Optional Target Platform requirements -->
 		<Property Name="InstalledRequirements" Value="$(InstalledRequirements);#UnrealGame LinuxNative;#UnrealGame LinuxNative Unstripped;#UnrealGame LinuxNative Stripped" If="'$(WithLinux)' == true"/>
@@ -1248,7 +1237,6 @@
 			<Property Name="CopyInstalledExceptions" Value="$(CopyEditorExceptionsLinux)"/>
 
 			<!-- Tag files that always need to be part of an installed build -->
-			<Tag Files="#Compile UnrealHeaderTool Linux" With="#Installed Linux"/>
 			<Tag Files="$(UnrealEditor)" Except="#UnrealEditor Linux Unstripped" With="#Installed Linux"/>
 			<Tag Files="#UnrealEditor Linux Stripped" With="#Saved Output"/>
 			<Tag Files="#Build Tools LinuxNative;#Make Feature Packs LinuxNative;#Rules Assemblies" With="#Installed Linux"/>
@@ -1292,7 +1280,7 @@
 			<Tag Files="$(CopyInstalledFilter)" With="#Installed Linux"/>
 
 			<!-- Tag Linux files that need to be stripped -->
-			<Tag Files="#Compile UnrealHeaderTool Linux;#Build Tools LinuxNative" Filter="$(LinuxStripFilter)" Except="$(CopyInstalledExceptions)" With="#Strip Tools for Installed Linux"/>
+			<Tag Files="#Build Tools LinuxNative" Filter="$(LinuxStripFilter)" Except="$(CopyInstalledExceptions)" With="#Strip Tools for Installed Linux"/>
 
 			<!-- Tag files that need to be copied directly -->
 			<Tag Files="#Installed Linux" Except="$(CopyInstalledExceptions);#Strip Tools for Installed Linux" With="#Copy for Installed Linux"/>
diff --git a/Engine/Platforms/Hololens/Build/InstalledEngineBuild.xml b/Engine/Platforms/Hololens/Build/InstalledEngineBuild.xml
--- a/Engine/Platforms/Hololens/Build/InstalledEngineBuild.xml
+++ b/Engine/Platforms/Hololens/Build/InstalledEngineBuild.xml
@@ -31,7 +31,7 @@
 
 	<Extend Name="Platform_Agent_Target_Win64_Licensee">
 		<Agent Name="Target HoloLens" Type="Win64_Licensee">
-			<Node Name="Compile UnrealGame HoloLens" Requires="Compile UnrealHeaderTool Win64" Produces="#UnrealGame HoloLens;#UnrealGame HoloLens Unstripped;#UnrealGame HoloLens Stripped;#UnrealGame HoloLens Unsigned;#UnrealGame HoloLens Signed">
+			<Node Name="Compile UnrealGame HoloLens" Requires="Update Version Files" Produces="#UnrealGame HoloLens;#UnrealGame HoloLens Unstripped;#UnrealGame HoloLens Stripped;#UnrealGame HoloLens Unsigned;#UnrealGame HoloLens Signed">
 				<ForEach Name="Configuration" Values="$(GameConfigurations)">
 					<Compile Target="UnrealGame" Platform="HoloLens" Configuration="$(Configuration)" Tag="#UnrealGame HoloLens" Arguments="-precompile -allmodules -nolink -architecture=arm64 $(TargetDebugInfoArg)"/>
 					<Compile Target="UnrealGame" Platform="HoloLens" Configuration="$(Configuration)" Tag="#UnrealGame HoloLens" Arguments="-precompile -architecture=arm64 $(TargetDebugInfoArg)" Clean="false"/>
