Move installer file tagging from inside a plugin definition to the installed engine build script

Also ensuring it is appropriately signed.

#rnx

diff --git a/Engine/Build/InstalledEngineFilters.xml b/Engine/Build/InstalledEngineFilters.xml
--- a/Engine/Build/InstalledEngineFilters.xml
+++ b/Engine/Build/InstalledEngineFilters.xml
@@ -253,6 +253,7 @@
 		Engine/Binaries/DotNET/CsvTools/...
 		Engine/Binaries/ThirdParty/DotNet/6.0.302/windows/...
 		Engine/Binaries/ThirdParty/Python3/...
+		Engine/Binaries/ThirdParty/PresentMon/Win64/PresentMon64-1.5.2.exe
 		Engine/Binaries/ThirdParty/svn/Win64/...
 		Engine/Binaries/Win64/UnrealTraceServer.exe
 		Engine/Binaries/Win64/EpicWebHelper.exe
@@ -311,6 +312,7 @@
 
 		<!-- Used by Switchboard application, and by iOS/tvOS build automation -->
 		Engine/Extras/ThirdPartyNotUE/cwrsync/...
+		Engine/Binaries/ThirdParty/PresentMon/Win64/PresentMon64-1.5.2.exe
 
 		<!-- nanosvg -->
 		Engine/Source/ThirdParty/nanosvg/...
diff --git a/Engine/Source/Programs/SwitchboardListener/SwitchboardListener.Build.cs b/Engine/Source/Programs/SwitchboardListener/SwitchboardListener.Build.cs
--- a/Engine/Source/Programs/SwitchboardListener/SwitchboardListener.Build.cs
+++ b/Engine/Source/Programs/SwitchboardListener/SwitchboardListener.Build.cs
@@ -28,9 +28,6 @@ public class SwitchboardListener : ModuleRules
 			);
 
 			PublicSystemLibraries.Add("Pdh.lib");
-
-			// Add PresentMon as a Runtime Dependency to avoid it from being filtered out when generating the release binaries.
-			RuntimeDependencies.Add("$(EngineDir)/Binaries/ThirdParty/PresentMon/Win64/PresentMon-1.8.0-x64.exe");
 		}
 
 		PrivateIncludePathModuleNames.AddRange(
