Update PresentMon reference

#rnx

diff --git a/Engine/Build/InstalledEngineFilters.xml b/Engine/Build/InstalledEngineFilters.xml
--- a/Engine/Build/InstalledEngineFilters.xml
+++ b/Engine/Build/InstalledEngineFilters.xml
@@ -253,7 +253,7 @@
 		Engine/Binaries/DotNET/CsvTools/...
 		Engine/Binaries/ThirdParty/DotNet/6.0.302/windows/...
 		Engine/Binaries/ThirdParty/Python3/...
-		Engine/Binaries/ThirdParty/PresentMon/Win64/PresentMon64-1.5.2.exe
+		Engine/Binaries/ThirdParty/PresentMon/Win64/PresentMon-1.8.0-x64.exe
 		Engine/Binaries/ThirdParty/svn/Win64/...
 		Engine/Binaries/Win64/UnrealTraceServer.exe
 		Engine/Binaries/Win64/EpicWebHelper.exe
@@ -312,7 +312,7 @@
 
 		<!-- Used by Switchboard application, and by iOS/tvOS build automation -->
 		Engine/Extras/ThirdPartyNotUE/cwrsync/...
-		Engine/Binaries/ThirdParty/PresentMon/Win64/PresentMon64-1.5.2.exe
+		Engine/Binaries/ThirdParty/PresentMon/Win64/PresentMon-1.8.0-x64.exe
 
 		<!-- nanosvg -->
 		Engine/Source/ThirdParty/nanosvg/...
