Update build.version to 5.3

diff --git a/Engine/Build/Build.version b/Engine/Build/Build.version
--- a/Engine/Build/Build.version
+++ b/Engine/Build/Build.version
@@ -1,6 +1,6 @@
 {
 	"MajorVersion": 5,
-	"MinorVersion": 2,
+	"MinorVersion": 3,
 	"PatchVersion": 0,
 	"Changelist": 0,
 	"CompatibleChangelist": 0,
