Changing Mac to generate Universal Binaries

#rnx

diff --git a/Engine/Build/InstalledEngineBuild.xml b/Engine/Build/InstalledEngineBuild.xml
--- a/Engine/Build/InstalledEngineBuild.xml
+++ b/Engine/Build/InstalledEngineBuild.xml
@@ -34,7 +34,6 @@
 	<!-- Setup default target platforms for Installed build if they haven't been specified on the commandline -->
 	<Option Name="WithWin64" Restrict="true|false" DefaultValue="$(DefaultWithWin64)" Description="Include the Win64 target platform"/>
 	<Option Name="WithMac" Restrict="true|false" DefaultValue="$(DefaultWithMac)" Description="Include the Mac target platform"/>
-	<Option Name="WithMacArm64" Restrict="true|false" DefaultValue="false" Description="Include the Arm64 architecture version of the Mac target platform"/>
 	<Option Name="WithAndroid" Restrict="true|false" DefaultValue="$(DefaultWithPlatform)" Description="Include the Android target platform"/>
 	<Option Name="WithIOS" Restrict="true|false" DefaultValue="$(DefaultWithIOS)" Description="Include the iOS target platform"/>
 	<Option Name="WithTVOS" Restrict="true|false" DefaultValue="$(DefaultWithIOS)" Description="Include the tvOS target platform"/>
@@ -142,8 +141,7 @@
 	<Property Name="CrashReporterCompileArgs" Value=""/>
 	<Property Name="CrashReporterCompileArgs" Value="-define:CRC_TELEMETRY_URL=\&quot;$(CrashReporterAPIURL)\&quot; -define:CRC_TELEMETRY_KEY_DEV=\&quot;$(CrashReporterAPIKey)\&quot; -define:CRC_TELEMETRY_KEY_RELEASE=\&quot;$(CrashReporterAPIKey)\&quot; -OverrideBuildEnvironment" If="'$(CrashReporterAPIURL)' != '' And '$(CrashReporterAPIKey)' != ''" />
 
-	<Property Name="UseArm64ArchitectureArg" Value="" />
-	<Property Name="UseArm64ArchitectureArg" Value="-architecture=arm64" If="$(WithMacArm64)" />
+	<Property Name="MacArchitectureArg" Value="-architecture=arm64+x64" />
 
 	<!-- Base editor and tools for compile on Windows -->
 	<Agent Name="Editor Win64" Type="Win64_Licensee">
@@ -191,9 +189,8 @@
 	<Agent Name="Editor Mac" Type="Mac_Licensee">
 		<Node Name="Compile UnrealEditor Mac" Requires="Update Version Files" Produces="#UnrealEditor Mac;#UnrealEditor Mac Unstripped;#UnrealEditor Mac Stripped;#UnrealEditor Mac Unsigned;#UnrealEditor Mac Signed;#UnrealEditor Mac WithExceptions">
 			<Property Name="AllModulesArg" Value="-allmodules"/>
-			<Property Name="AllModulesArg" Value="" If="$(WithMacArm64)"/>
-			<Compile Target="UnrealEditor" Platform="Mac" Configuration="DebugGame" Tag="#UnrealEditor Mac" Arguments="-precompile -EnableDSYM $(AllModulesArg) $(CrashReporterCompileArgs) $(UseArm64ArchitectureArg)"/>
-			<Compile Target="UnrealEditor" Platform="Mac" Configuration="Development" Tag="#UnrealEditor Mac" Arguments="-precompile -EnableDSYM $(AllModulesArg) $(CrashReporterCompileArgs) $(UseArm64ArchitectureArg)"/>
+			<Compile Target="UnrealEditor" Platform="Mac" Configuration="DebugGame" Tag="#UnrealEditor Mac" Arguments="-precompile -EnableDSYM $(AllModulesArg) $(CrashReporterCompileArgs) $(MacArchitectureArg)"/>
+			<Compile Target="UnrealEditor" Platform="Mac" Configuration="Development" Tag="#UnrealEditor Mac" Arguments="-precompile -EnableDSYM $(AllModulesArg) $(CrashReporterCompileArgs) $(MacArchitectureArg)"/>
 
 			<Property Name="UnrealEditorExceptions" Value=""/>
 			<Property Name="UnrealEditorExceptions" Value="$(PluginsExceptions)" If="$(CompileDatasmithPlugins)"/>
@@ -812,7 +809,6 @@
 	</Do>
 
 	<Property Name="MacDDCAgentType" Value="Mac" />
-	<Property Name="MacDDCAgentType" Value="MacArm64" If="$(WithMacArm64)" />
 	<Agent Name="DDC Group Mac" Type="$(MacDDCAgentType)">
 		<Node Name="Build DDC Mac" Requires="#UnrealEditor Mac;Build Tools Mac;#Build Tools CS Binaries Mac" Produces="#CompressedDDCMac">
 			<!-- Get our temp dir -->
@@ -905,7 +901,7 @@
 			<Do If="$(WithMac) And $(CompileDatasmithPlugins) == 'True'">
 				<Property Name="DatasmithPlugins" Value="DatasmithCADImporter;DatasmithIFCImporter;DatasmithC4DImporter;AxFImporter;MDLImporter"/>
 				<ForEach Name="Plugin" Values="$(DatasmithPlugins)">
-					<Compile Target="UnrealEditor" Configuration="Development" Platform="Mac" Arguments="-Plugin=&quot;$(RootDir)/Engine/Plugins/Enterprise/$(Plugin)/$(Plugin).uplugin&quot; $(UseArm64ArchitectureArg)"/>
+					<Compile Target="UnrealEditor" Configuration="Development" Platform="Mac" Arguments="-Plugin=&quot;$(RootDir)/Engine/Plugins/Enterprise/$(Plugin)/$(Plugin).uplugin&quot; $(MacArchitectureArg)"/>
 					<Tag Files="$(RootDir)/Engine/Plugins/Enterprise/$(Plugin)/Binaries/Mac/..." With="#DatasmithPluginBinaries Mac"/>
 				</ForEach>
 			</Do>
