Fixing graph error when WithWin64 is false

#rnx

diff --git a/Engine/Build/InstalledEngineBuild.xml b/Engine/Build/InstalledEngineBuild.xml
--- a/Engine/Build/InstalledEngineBuild.xml
+++ b/Engine/Build/InstalledEngineBuild.xml
@@ -1033,7 +1033,7 @@
 				<Tag Files="$(CopyInstalledFilter)" Filter="$(WindowsSignFilter)" Except="$(SignCopyExceptions)" With="#Sign Tools for Installed Win64"/>
 				<Tag FileLists="#DependencyLists" Filter="$(WindowsSignFilter)" Except="$(SignCopyExceptions)" With="#Sign Tools for Installed Win64"/>
 				<Tag Files="$(ExtraSignFiles)" Filter="$(WindowsSignFilter)" Except="$(SignCopyExceptions)" With="#Sign Tools for Installed Win64"/>
-				<Tag Files="#DatasmithPluginBinaries Win64" Filter="$(WindowsSignFilter)" Except="$(SignCopyExceptions)" With="#Sign Tools for Installed Win64" If="$(CompileDatasmithPlugins)"/>
+				<Tag Files="#DatasmithPluginBinaries Win64" Filter="$(WindowsSignFilter)" Except="$(SignCopyExceptions)" With="#Sign Tools for Installed Win64" If="$(WithWin64) And $(CompileDatasmithPlugins)"/>
 				<Tag Files="#Build Tools CS Binaries;" Filter="$(WindowsSignFilter)" Except="$(CopyWin64CsToolsExceptions);" With="#Sign CS Tools for Installed Win64"/>
 			</Do>
 
