Build: Filter out python files for platforms that can't use them in InstallBuilds.

#10137

diff --git a/Engine/Build/InstalledEngineFilters.xml b/Engine/Build/InstalledEngineFilters.xml
--- a/Engine/Build/InstalledEngineFilters.xml
+++ b/Engine/Build/InstalledEngineFilters.xml
@@ -363,6 +363,10 @@
 
 		<!-- Don't want these folders, even if they're part of Windows tools -->
 		Engine/Binaries/Win64/...
+
+		<!-- Don't want The Windows/Linux Python packages -->
+		Engine/Plugins/Experimental/PythonFoundationPackages/Content/Python/Lib/Win64/...
+		Engine/Plugins/Experimental/PythonFoundationPackages/Content/Python/Lib/Linux/...
 	</Property>
 
 	<!-- Linux editor specific dependencies -->
@@ -414,6 +418,10 @@
 		<!-- Don't want these folders, even if they're part of Windows/Mac tools -->
 		Engine/Binaries/Win64/...
 		Engine/Binaries/Mac/...
+
+		<!-- Don't want The Windows/Mac Python packages -->
+		Engine/Plugins/Experimental/PythonFoundationPackages/Content/Python/Lib/Win64/...
+		Engine/Plugins/Experimental/PythonFoundationPackages/Content/Python/Lib/Mac/...
 	</Property>
 
 	<!-- Target Platform Filters/Exceptions -->
