Disable ReplicationSystem tests on non-Win64 platforms unti we can pass Win64 cleanly.

diff --git a/Engine/Build/LowLevelTests_GenProps.xml b/Engine/Build/LowLevelTests_GenProps.xml
--- a/Engine/Build/LowLevelTests_GenProps.xml
+++ b/Engine/Build/LowLevelTests_GenProps.xml
@@ -19,7 +19,7 @@
   <Option Name="RunReplicationSystemTests" DefaultValue="False" Description="Run Replication System Tests" />
   <Option Name="RunFoundationTests" DefaultValue="False" Description="Run Foundation Tests" />
   <Option Name="RunHeadlessChaosTests" DefaultValue="False" Description="Run Headless Chaos Tests" />
-  <Property Name="ReplicationSystemSupportedPlatforms" Value="Win64;Mac;Linux;Android" />
+  <Property Name="ReplicationSystemSupportedPlatforms" Value="Win64" />
   <Property Name="FoundationSupportedPlatforms" Value="Win64;Mac;Linux;Android" />
   <Property Name="HeadlessChaosSupportedPlatforms" Value="Win64" />
 </BuildGraph>
\ No newline at end of file
