Gauntlet - Use space instead of coma when parsing extra engine arguments
Follow up to 24507995

diff --git a/Engine/Build/EditorPerf.xml b/Engine/Build/EditorPerf.xml
--- a/Engine/Build/EditorPerf.xml
+++ b/Engine/Build/EditorPerf.xml
@@ -90,7 +90,7 @@
 		<Regex Pattern="(.*),(.*)" Capture="Project;MapAsset" Input="$(ProjectPair)" />
 
 		<Property Name="$(Project)MapArgs" Value=""/>
-		<Property Name="$(Project)MapArgs" Value="-ini:Engine:[/Script/Engine.AutomationTestSettings]:MapsToPIETest=&quot;$(MapAsset)&quot;" If="'$(MapAsset)' != ''" />
+		<Property Name="$(Project)MapArgs" Value="-ini:Engine:[/Script/Engine.AutomationTestSettings]:MapsToPIETest=\&quot;$(MapAsset)\&quot;" If="'$(MapAsset)' != ''" />
 	</ForEach>
 
 	<!-- Options for extra Zen Loader tests -->
@@ -100,17 +100,17 @@
 
 	<!-- Zen Loader tests specifc arguments -->
 	<Property Name="ZenLoaderTestArgs" Value="-zenloader"/>
-	<Property Name="ZenLoaderAsyncTestArgs" Value="$(ZenLoaderTestArgs),-asyncloadingthread" />
+	<Property Name="ZenLoaderAsyncTestArgs" Value="$(ZenLoaderTestArgs) -asyncloadingthread" />
 	
-	<Property Name="ZenLoaderTestArgs" Value="$(ZenLoaderTestArgs),-noeditordomain" If="$(EnableZenTestsWithoutEditorDomain)" />
-	<Property Name="ZenLoaderAsyncTestArgs" Value="$(ZenLoaderAsyncTestArgs),-noeditordomain" If="$(EnableZenTestsWithoutEditorDomain)" />
+	<Property Name="ZenLoaderTestArgs" Value="$(ZenLoaderTestArgs) -noeditordomain" If="$(EnableZenTestsWithoutEditorDomain)" />
+	<Property Name="ZenLoaderAsyncTestArgs" Value="$(ZenLoaderAsyncTestArgs) -noeditordomain" If="$(EnableZenTestsWithoutEditorDomain)" />
 
 	<!-- Options for extra Zen Store tests -->
 	<Option Name="EnableZenStoreDDCTest" DefaultValue="false" Description="Enable additional test on warm phase that uses the Zen Store DDC."/>
 
 	<!-- Zen Store tests specific arguments -->
 	<Property Name="ZenServerCache" Value="$(RootDir)\ZenCache" />
-	<Property Name="ZenStoreCommonArgs" Value="-ini:Engine:[Zen.AutoLaunch]:LimitProcessLifetime=true,-ZenDataPath=$(ZenServerCache) "/>
+	<Property Name="ZenStoreCommonArgs" Value="-ini:Engine:[Zen.AutoLaunch]:LimitProcessLifetime=true -ZenDataPath=$(ZenServerCache) "/>
 	<Property Name="ZenStoreDDCPhaseArgs" Value="-ddc=ZenDDC"/>
 	<Property Name="ZenStoreDDCTestArgs" Value="$(ZenStoreCommonArgs)"/>
 
@@ -120,7 +120,7 @@
 	
 	<!-- LazyLoad tests specific arguments -->
 	<Property Name="LazyLoadTestArgs" Value="-LazyLoadImports"/>
-	<Property Name="LazyLoadZenTestArgs" Value="$(LazyLoadTestArgs),-zenloader" />
+	<Property Name="LazyLoadZenTestArgs" Value="$(LazyLoadTestArgs) -zenloader" />
 
 	<!-- Gauntlet arguments for all runs -->
 	<Property Name="GauntletArgs" Value="-build=editor -UseEditor -verbose -tempdir=&quot;$(RootDir)/Tests&quot;"/>
@@ -256,20 +256,20 @@
 						<!-- Zen Loader TestRun if enabled -->
 						<Property Name="RunNodeZenName" Value="Run $(OutputName) $(Phase) Zen"/>
 						<Node Name="$(RunNodeZenName)" Requires="$(RunNodeRequires);$(RunNodeList)" If="$(EnableZenLoaderTest)">
-							<Expand Name="RunTest" TestNameSuffix="Zen" Platform="$(Platform)" Project="$(Project)" NodeName="$(RunNodeZenName)" Phase="$(Phase)" Iteration="1" Args="$($(Phase)Args)" EditorArgs="$(ProjectMapArgs),$(ZenLoaderTestArgs)"/>
+							<Expand Name="RunTest" TestNameSuffix="Zen" Platform="$(Platform)" Project="$(Project)" NodeName="$(RunNodeZenName)" Phase="$(Phase)" Iteration="1" Args="$($(Phase)Args)" EditorArgs="$(ProjectMapArgs) $(ZenLoaderTestArgs)"/>
 						</Node>
 						<Property Name="RunNodeList" Value="$(RunNodeList);$(RunNodeZenName)" If="$(EnableZenLoaderTest)"/>
 						<!-- Zen Loader Async TestRun if enabled -->
 						<Property Name="RunNodeZenAsyncName" Value="Run $(OutputName) $(Phase) Zen Async"/>
 						<Node Name="$(RunNodeZenAsyncName)" Requires="$(RunNodeRequires);$(RunNodeList)" If="$(EnableZenAsyncThreadTest)">
-							<Expand Name="RunTest" TestNameSuffix="ZenAsync" Platform="$(Platform)" Project="$(Project)" NodeName="$(RunNodeZenAsyncName)" Phase="$(Phase)" Iteration="1" Args="$($(Phase)Args)" EditorArgs="$(ProjectMapArgs),$(ZenLoaderAsyncTestArgs)"/>
+							<Expand Name="RunTest" TestNameSuffix="ZenAsync" Platform="$(Platform)" Project="$(Project)" NodeName="$(RunNodeZenAsyncName)" Phase="$(Phase)" Iteration="1" Args="$($(Phase)Args)" EditorArgs="$(ProjectMapArgs) $(ZenLoaderAsyncTestArgs)"/>
 						</Node>
 						<Property Name="RunNodeList" Value="$(RunNodeList);$(RunNodeZenAsyncName)" If="$(EnableZenAsyncThreadTest)"/>
 					
 						<!-- Zen Store DDC TestRun if enabled -->
 						<Property Name="RunNodeZenStoreDDCName" Value="Run $(OutputName) $(Phase) Zen Store DDC"/>
 						<Node Name="$(RunNodeZenStoreDDCName)" Requires="$(RunNodeRequires);$(RunNodeList)" If="$(EnableZenStoreDDCTest)">
-							<Expand Name="RunTest" TestNameSuffix="ZenAsync" Platform="$(Platform)" Project="$(Project)" NodeName="$(RunNodeZenStoreDDCName)" Phase="$(Phase)" Iteration="1" Args="$(ZenStoreDDCPhaseArgs)" EditorArgs="$(ProjectMapArgs),$(ZenStoreDDCTestArgs)"/>
+							<Expand Name="RunTest" TestNameSuffix="ZenAsync" Platform="$(Platform)" Project="$(Project)" NodeName="$(RunNodeZenStoreDDCName)" Phase="$(Phase)" Iteration="1" Args="$(ZenStoreDDCPhaseArgs)" EditorArgs="$(ProjectMapArgs) $(ZenStoreDDCTestArgs)"/>
 						</Node>
 						<Property Name="RunNodeList" Value="$(RunNodeList);$(RunNodeZenStoreDDCName)" If="$(EnableZenStoreDDCTest)"/>
 					</Do>
@@ -278,13 +278,13 @@
 						<!-- Legacy LazyLoad TestRun if enabled -->
 						<Property Name="RunNodeLazyLoadName" Value="Run $(OutputName) $(Phase) LazyLoad"/>
 						<Node Name="$(RunNodeLazyLoadName)" Requires="$(RunNodeRequires);$(RunNodeList)" If="$(EnableLazyLoadTest)">
-							<Expand Name="RunTest" TestNameSuffix="Zen" Platform="$(Platform)" Project="$(Project)" NodeName="$(RunNodeLazyLoadName)" Phase="$(Phase)" Iteration="1" Args="$($(Phase)Args)" EditorArgs="$(ProjectMapArgs),$(LazyLoadTestArgs)"/>
+							<Expand Name="RunTest" TestNameSuffix="Zen" Platform="$(Platform)" Project="$(Project)" NodeName="$(RunNodeLazyLoadName)" Phase="$(Phase)" Iteration="1" Args="$($(Phase)Args)" EditorArgs="$(ProjectMapArgs) $(LazyLoadTestArgs)"/>
 						</Node>
 						<Property Name="RunNodeList" Value="$(RunNodeList);$(RunNodeLazyLoadName)" If="$(EnableLazyLoadTest)"/>
 						<!-- Zen LazyLoad TestRun if enabled -->
 						<Property Name="RunNodeLazyLoadZenName" Value="Run $(OutputName) $(Phase) LazyLoad Zen"/>
 						<Node Name="$(RunNodeLazyLoadZenName)" Requires="$(RunNodeRequires);$(RunNodeList)" If="$(EnableLazyLoadZenTest)">
-							<Expand Name="RunTest" TestNameSuffix="Zen" Platform="$(Platform)" Project="$(Project)" NodeName="$(RunNodeLazyLoadZenName)" Phase="$(Phase)" Iteration="1" Args="$($(Phase)Args)" EditorArgs="$(ProjectMapArgs),$(LazyLoadZenTestArgs)"/>
+							<Expand Name="RunTest" TestNameSuffix="Zen" Platform="$(Platform)" Project="$(Project)" NodeName="$(RunNodeLazyLoadZenName)" Phase="$(Phase)" Iteration="1" Args="$($(Phase)Args)" EditorArgs="$(ProjectMapArgs) $(LazyLoadZenTestArgs)"/>
 						</Node>
 						<Property Name="RunNodeList" Value="$(RunNodeList);$(RunNodeLazyLoadZenName)" If="$(EnableLazyLoadZenTest)"/>
 					</Do>
diff --git a/Engine/Source/Programs/AutomationTool/Gauntlet/Unreal/RunUnreal.cs b/Engine/Source/Programs/AutomationTool/Gauntlet/Unreal/RunUnreal.cs
--- a/Engine/Source/Programs/AutomationTool/Gauntlet/Unreal/RunUnreal.cs
+++ b/Engine/Source/Programs/AutomationTool/Gauntlet/Unreal/RunUnreal.cs
@@ -213,7 +213,7 @@ namespace Gauntlet
 					}
 
 					// look for-args= and then -clientargs= and -editorargs etc
-					Role.ExtraArgs = string.Join(' ', Globals.Params.ParseValues("Args", true));
+					Role.ExtraArgs = string.Join(' ', Globals.Params.ParseValue("Args", string.Empty));
 					string ExtraRoleArgs = Globals.Params.ParseValue(Type.ToString() + "Args", string.Empty);
 					if (!string.IsNullOrEmpty(ExtraRoleArgs))
 					{
