Add support for Linux
Add support to pass arguments to editor process with an option
Add support for compiling with TSAN
#rnx

diff --git a/Engine/Build/EditorPerf.xml b/Engine/Build/EditorPerf.xml
--- a/Engine/Build/EditorPerf.xml
+++ b/Engine/Build/EditorPerf.xml
@@ -24,10 +24,12 @@
 	<Option Name="Platforms" DefaultValue="Win64" Description="The list of platforms to profile"/>
 	<Property Name="Win64CompileAgentType" Value="Win64"/>
 	<Property Name="MacCompileAgentType" Value="TestMac"/>
+	<Property Name="LinuxCompileAgentType" Value="Linux"/>
 	<Error Message="No platforms specified" If="'$(Platforms)' == ''"/>
 
 	<!-- Build Options -->
 	<Option Name="NoDebugInfo" Restrict="true|false" DefaultValue="false" Description="Compile without including debug info"/>
+	<Option Name="WithTSAN" Restrict="true|false" DefaultValue="false" Description="Compile using TSAN"/>
 	<Property Name="ExtraPluginsCommaDelims" Value="EditorTests"/>
 	<Property Name="ExtraPluginsPlusDelims" Value=""/> <!-- populated programmatically -->
 	<ForEach Name="Plugin" Values="$(ExtraPluginsCommaDelims)">
@@ -38,6 +40,7 @@
 	<Property Name="EditorCompileArgs" Value=""/>
 	<Property Name="EditorCompileArgs" Value="-BuildPlugin=$(ExtraPluginsPlusDelims)" If="'$(ExtraPluginsPlusDelims)' != ''"/>
 	<Property Name="EditorCompileArgs" Value="$(EditorCompileArgs) -nodebuginfo" If="$(NoDebugInfo)"/>
+	<Property Name="EditorCompileArgs" Value="$(EditorCompileArgs) -EnableTSAN -Monolithic" If="$(WithTSAN)"/>
 
 	<!-- Automation uses its own network output directory -->
 	<Property Name="AutomationNetworkRootDir" Value="\\epicgames.net\root\Builds\Automation\Reports\EditorPerf"/>
@@ -92,6 +95,9 @@
 		<Property Name="$(Project)MapArgs" Value=""/>
 		<Property Name="$(Project)MapArgs" Value="-ini:Engine:[/Script/Engine.AutomationTestSettings]:MapsToPIETest=\&quot;$(MapAsset)\&quot;" If="'$(MapAsset)' != ''" />
 	</ForEach>
+	
+	<!-- Options for extra runtime args -->
+	<Option Name="AdditionalRuntimeArgs" DefaultValue="" Description="Additional Arguments to pass to RunTest"/>
 
 	<!-- Options for extra Zen Loader tests -->
 	<Option Name="EnableZenLoaderTest" DefaultValue="true" Description="Enable additional test per phase that uses the Zen Loader."/>
@@ -213,7 +219,7 @@
 			<Regex Pattern="(.*)\[(.*)\]" Capture="Project;EditorTarget" Input="$(Project)" Optional="true"/>
 
 			<!-- Project Map Args-->
-			<Property Name="ProjectMapArgs" Value="$($(Project)MapArgs)"/>
+			<Property Name="RuntimeArgs" Value="$($(Project)MapArgs) $(AdditionalRuntimeArgs)"/>
 
 			<!-- Use as base string for directories and key for node names -->
 			<Property Name="OutputName" Value="$(Project)$(Platform)EditorPerf"/>
@@ -245,7 +251,7 @@
 								<Property Name="IterationSuffix" Value="$(Iteration)" If="'$(Iteration)' != '1'"/>
 								<Property Name="RunNodeName" Value="Run $(OutputName) $(Phase)$(IterationSuffix)"/>
 								<Node Name="$(RunNodeName)" Requires="$(RunNodeRequires);$(RunNodeList)">
-									<Expand Name="RunTest" TestNameSuffix="" Platform="$(Platform)" Project="$(Project)" NodeName="$(RunNodeName)" Phase="$(Phase)" Iteration="$(Iteration)" Args="$($(Phase)Args)" EditorArgs="$(ProjectMapArgs)"/>
+									<Expand Name="RunTest" TestNameSuffix="" Platform="$(Platform)" Project="$(Project)" NodeName="$(RunNodeName)" Phase="$(Phase)" Iteration="$(Iteration)" Args="$($(Phase)Args)" EditorArgs="$(RuntimeArgs)"/>
 								</Node>
 								<Property Name="RunNodeList" Value="$(RunNodeList);$(RunNodeName)"/>
 							</Do>
@@ -256,20 +262,20 @@
 						<!-- Zen Loader TestRun if enabled -->
 						<Property Name="RunNodeZenName" Value="Run $(OutputName) $(Phase) Zen"/>
 						<Node Name="$(RunNodeZenName)" Requires="$(RunNodeRequires);$(RunNodeList)" If="$(EnableZenLoaderTest)">
-							<Expand Name="RunTest" TestNameSuffix="Zen" Platform="$(Platform)" Project="$(Project)" NodeName="$(RunNodeZenName)" Phase="$(Phase)" Iteration="1" Args="$($(Phase)Args)" EditorArgs="$(ProjectMapArgs) $(ZenLoaderTestArgs)"/>
+							<Expand Name="RunTest" TestNameSuffix="Zen" Platform="$(Platform)" Project="$(Project)" NodeName="$(RunNodeZenName)" Phase="$(Phase)" Iteration="1" Args="$($(Phase)Args)" EditorArgs="$(RuntimeArgs) $(ZenLoaderTestArgs)"/>
 						</Node>
 						<Property Name="RunNodeList" Value="$(RunNodeList);$(RunNodeZenName)" If="$(EnableZenLoaderTest)"/>
 						<!-- Zen Loader Async TestRun if enabled -->
 						<Property Name="RunNodeZenAsyncName" Value="Run $(OutputName) $(Phase) Zen Async"/>
 						<Node Name="$(RunNodeZenAsyncName)" Requires="$(RunNodeRequires);$(RunNodeList)" If="$(EnableZenAsyncThreadTest)">
-							<Expand Name="RunTest" TestNameSuffix="ZenAsync" Platform="$(Platform)" Project="$(Project)" NodeName="$(RunNodeZenAsyncName)" Phase="$(Phase)" Iteration="1" Args="$($(Phase)Args)" EditorArgs="$(ProjectMapArgs) $(ZenLoaderAsyncTestArgs)"/>
+							<Expand Name="RunTest" TestNameSuffix="ZenAsync" Platform="$(Platform)" Project="$(Project)" NodeName="$(RunNodeZenAsyncName)" Phase="$(Phase)" Iteration="1" Args="$($(Phase)Args)" EditorArgs="$(RuntimeArgs) $(ZenLoaderAsyncTestArgs)"/>
 						</Node>
 						<Property Name="RunNodeList" Value="$(RunNodeList);$(RunNodeZenAsyncName)" If="$(EnableZenAsyncThreadTest)"/>
 					
 						<!-- Zen Store DDC TestRun if enabled -->
 						<Property Name="RunNodeZenStoreDDCName" Value="Run $(OutputName) $(Phase) Zen Store DDC"/>
 						<Node Name="$(RunNodeZenStoreDDCName)" Requires="$(RunNodeRequires);$(RunNodeList)" If="$(EnableZenStoreDDCTest)">
-							<Expand Name="RunTest" TestNameSuffix="ZenAsync" Platform="$(Platform)" Project="$(Project)" NodeName="$(RunNodeZenStoreDDCName)" Phase="$(Phase)" Iteration="1" Args="$(ZenStoreDDCPhaseArgs)" EditorArgs="$(ProjectMapArgs) $(ZenStoreDDCTestArgs)"/>
+							<Expand Name="RunTest" TestNameSuffix="ZenAsync" Platform="$(Platform)" Project="$(Project)" NodeName="$(RunNodeZenStoreDDCName)" Phase="$(Phase)" Iteration="1" Args="$(ZenStoreDDCPhaseArgs)" EditorArgs="$(RuntimeArgs) $(ZenStoreDDCTestArgs)"/>
 						</Node>
 						<Property Name="RunNodeList" Value="$(RunNodeList);$(RunNodeZenStoreDDCName)" If="$(EnableZenStoreDDCTest)"/>
 					</Do>
@@ -278,13 +284,13 @@
 						<!-- Legacy LazyLoad TestRun if enabled -->
 						<Property Name="RunNodeLazyLoadName" Value="Run $(OutputName) $(Phase) LazyLoad"/>
 						<Node Name="$(RunNodeLazyLoadName)" Requires="$(RunNodeRequires);$(RunNodeList)" If="$(EnableLazyLoadTest)">
-							<Expand Name="RunTest" TestNameSuffix="Zen" Platform="$(Platform)" Project="$(Project)" NodeName="$(RunNodeLazyLoadName)" Phase="$(Phase)" Iteration="1" Args="$($(Phase)Args)" EditorArgs="$(ProjectMapArgs) $(LazyLoadTestArgs)"/>
+							<Expand Name="RunTest" TestNameSuffix="Zen" Platform="$(Platform)" Project="$(Project)" NodeName="$(RunNodeLazyLoadName)" Phase="$(Phase)" Iteration="1" Args="$($(Phase)Args)" EditorArgs="$(RuntimeArgs) $(LazyLoadTestArgs)"/>
 						</Node>
 						<Property Name="RunNodeList" Value="$(RunNodeList);$(RunNodeLazyLoadName)" If="$(EnableLazyLoadTest)"/>
 						<!-- Zen LazyLoad TestRun if enabled -->
 						<Property Name="RunNodeLazyLoadZenName" Value="Run $(OutputName) $(Phase) LazyLoad Zen"/>
 						<Node Name="$(RunNodeLazyLoadZenName)" Requires="$(RunNodeRequires);$(RunNodeList)" If="$(EnableLazyLoadZenTest)">
-							<Expand Name="RunTest" TestNameSuffix="Zen" Platform="$(Platform)" Project="$(Project)" NodeName="$(RunNodeLazyLoadZenName)" Phase="$(Phase)" Iteration="1" Args="$($(Phase)Args)" EditorArgs="$(ProjectMapArgs) $(LazyLoadZenTestArgs)"/>
+							<Expand Name="RunTest" TestNameSuffix="Zen" Platform="$(Platform)" Project="$(Project)" NodeName="$(RunNodeLazyLoadZenName)" Phase="$(Phase)" Iteration="1" Args="$($(Phase)Args)" EditorArgs="$(RuntimeArgs) $(LazyLoadZenTestArgs)"/>
 						</Node>
 						<Property Name="RunNodeList" Value="$(RunNodeList);$(RunNodeLazyLoadZenName)" If="$(EnableLazyLoadZenTest)"/>
 					</Do>
