add option to build with 2019 toolchain for win64 precompile targets

diff --git a/Engine/Build/InstalledEngineBuild.xml b/Engine/Build/InstalledEngineBuild.xml
--- a/Engine/Build/InstalledEngineBuild.xml
+++ b/Engine/Build/InstalledEngineBuild.xml
@@ -71,6 +71,10 @@
 	<!-- Control ParallelExecutor -->
 	<Option	Name="AllowParallelExecutor" Restrict="true|false" DefaultValue="false" Description="Allow the use of UAT's ParallelExecutor when compiling"/>
 	<Property Name="bAllowParallelExecutor" Value="$(AllowParallelExecutor)" />
+	
+	<Option Name="Use2019Toolchain" Restrict="true|false" DefaultValue="false" Description="Whether to build precompile targets with the vs2019 toolchain"/>
+	<Property Name="VS2019Arg" Value=""/>
+	<Property Name="VS2019Arg" Value="-2019" If="$(Use2019Toolchain)"/>
 
 	<!-- The local output directory -->
 	<Property Name="LocalInstalledDir" Value="$(BuiltDirectory)/Windows"/>
@@ -149,8 +153,8 @@
 			<SetVersion Change="$(Change)" Branch="$(EscapedBranch)" If="$(IsBuildMachine)"/>
 		</Node>
 		<Node Name="Compile UnrealEditor Win64" Requires="Update Version Files" Produces="#UnrealEditor Win64;#UnrealEditor Win64 Unstripped;#UnrealEditor Win64 Stripped;#UnrealEditor Win64 Unsigned;#UnrealEditor Win64 Signed;#UnrealEditor Win64 WithExceptions">
-			<Compile Target="UnrealEditor" Platform="Win64" Configuration="DebugGame" Tag="#UnrealEditor Win64" Arguments="-precompile -allmodules $(CrashReporterCompileArgs)"/>
-			<Compile Target="UnrealEditor" Platform="Win64" Configuration="Development" Tag="#UnrealEditor Win64" Arguments="-precompile -allmodules $(CrashReporterCompileArgs)"/>
+			<Compile Target="UnrealEditor" Platform="Win64" Configuration="DebugGame" Tag="#UnrealEditor Win64" Arguments="-precompile -allmodules $(CrashReporterCompileArgs) $(VS2019Arg)"/>
+			<Compile Target="UnrealEditor" Platform="Win64" Configuration="Development" Tag="#UnrealEditor Win64" Arguments="-precompile -allmodules $(CrashReporterCompileArgs) $(VS2019Arg)"/>
 
 			<Do If="$(EmbedSrcSrvInfo)">
 				<!-- Embed source info into the PDB files. Should be done from this machine to ensure that paths are correct. -->
@@ -247,8 +251,8 @@
 		<ForEach Name="Target" Values="$(OptionalClientTarget);$(OptionalServerTarget)">
 			<Node Name="Compile $(Target) Win64" Requires="Update Version Files" Produces="#$(Target) Win64">
 				<ForEach Name="Configuration" Values="$(GameConfigurations)">
-					<Compile Target="$(Target)" Platform="Win64" Configuration="$(Configuration)" Tag="#$(Target) Win64" Arguments="-precompile -allmodules -nolink $(TargetDebugInfoArg)"/>
-					<Compile Target="$(Target)" Platform="Win64" Configuration="$(Configuration)" Tag="#$(Target) Win64" Arguments="-precompile $(TargetDebugInfoArg)" Clean="false"/>
+					<Compile Target="$(Target)" Platform="Win64" Configuration="$(Configuration)" Tag="#$(Target) Win64" Arguments="-precompile -allmodules -nolink $(TargetDebugInfoArg) $(VS2019Arg)"/>
+					<Compile Target="$(Target)" Platform="Win64" Configuration="$(Configuration)" Tag="#$(Target) Win64" Arguments="-precompile $(TargetDebugInfoArg) $(VS2019Arg)" Clean="false"/>
 					<Property Name="OptionalTargetTags" Value="$(OptionalTargetTags);#$(Target) Win64"/>
 				</ForEach>
 			</Node>
@@ -258,8 +262,8 @@
 				<Tag Files="$(Tag)" With="#UnrealGame Win64"/>
 			</ForEach>
 			<ForEach Name="Configuration" Values="$(GameConfigurations)">
-				<Compile Target="UnrealGame" Platform="Win64" Configuration="$(Configuration)" Tag="#UnrealGame Win64" AllowParallelExecutor="$(bAllowParallelExecutor)" Arguments="-precompile -allmodules -nolink $(TargetDebugInfoArg)"/>
-				<Compile Target="UnrealGame" Platform="Win64" Configuration="$(Configuration)" Tag="#UnrealGame Win64" AllowParallelExecutor="$(bAllowParallelExecutor)" Arguments="-precompile $(TargetDebugInfoArg)" Clean="false"/>
+				<Compile Target="UnrealGame" Platform="Win64" Configuration="$(Configuration)" Tag="#UnrealGame Win64" AllowParallelExecutor="$(bAllowParallelExecutor)" Arguments="-precompile -allmodules -nolink $(TargetDebugInfoArg) $(VS2019Arg)"/>
+				<Compile Target="UnrealGame" Platform="Win64" Configuration="$(Configuration)" Tag="#UnrealGame Win64" AllowParallelExecutor="$(bAllowParallelExecutor)" Arguments="-precompile $(TargetDebugInfoArg) $(VS2019Arg)" Clean="false"/>
 			</ForEach>
 
 			<Do If="$(EmbedSrcSrvInfo)">
diff --git a/Engine/Platforms/Hololens/Build/InstalledEngineBuild.xml b/Engine/Platforms/Hololens/Build/InstalledEngineBuild.xml
--- a/Engine/Platforms/Hololens/Build/InstalledEngineBuild.xml
+++ b/Engine/Platforms/Hololens/Build/InstalledEngineBuild.xml
@@ -33,10 +33,10 @@
 		<Agent Name="Target HoloLens" Type="Win64_Licensee">
 			<Node Name="Compile UnrealGame HoloLens" Requires="Update Version Files" Produces="#UnrealGame HoloLens;#UnrealGame HoloLens Unstripped;#UnrealGame HoloLens Stripped;#UnrealGame HoloLens Unsigned;#UnrealGame HoloLens Signed">
 				<ForEach Name="Configuration" Values="$(GameConfigurations)">
-					<Compile Target="UnrealGame" Platform="HoloLens" Configuration="$(Configuration)" Tag="#UnrealGame HoloLens" Arguments="-precompile -allmodules -nolink -architecture=arm64 $(TargetDebugInfoArg)"/>
-					<Compile Target="UnrealGame" Platform="HoloLens" Configuration="$(Configuration)" Tag="#UnrealGame HoloLens" Arguments="-precompile -architecture=arm64 $(TargetDebugInfoArg)" Clean="false"/>
-					<Compile Target="UnrealGame" Platform="HoloLens" Configuration="$(Configuration)" Tag="#UnrealGame HoloLens" Arguments="-precompile -allmodules -nolink -architecture=x64 $(TargetDebugInfoArg)"/>
-					<Compile Target="UnrealGame" Platform="HoloLens" Configuration="$(Configuration)" Tag="#UnrealGame HoloLens" Arguments="-precompile -architecture=x64 $(TargetDebugInfoArg)" Clean="false"/>
+					<Compile Target="UnrealGame" Platform="HoloLens" Configuration="$(Configuration)" Tag="#UnrealGame HoloLens" Arguments="-precompile -allmodules -nolink -architecture=arm64 $(TargetDebugInfoArg) $(VS2019Arg)"/>
+					<Compile Target="UnrealGame" Platform="HoloLens" Configuration="$(Configuration)" Tag="#UnrealGame HoloLens" Arguments="-precompile -architecture=arm64 $(TargetDebugInfoArg) $(VS2019Arg)" Clean="false"/>
+					<Compile Target="UnrealGame" Platform="HoloLens" Configuration="$(Configuration)" Tag="#UnrealGame HoloLens" Arguments="-precompile -allmodules -nolink -architecture=x64 $(TargetDebugInfoArg) $(VS2019Arg)"/>
+					<Compile Target="UnrealGame" Platform="HoloLens" Configuration="$(Configuration)" Tag="#UnrealGame HoloLens" Arguments="-precompile -architecture=x64 $(TargetDebugInfoArg) $(VS2019Arg)" Clean="false"/>
 				</ForEach>
 
 				<Switch>
