Remove stray debug warning

diff --git a/Engine/Build/EditorPerf.xml b/Engine/Build/EditorPerf.xml
--- a/Engine/Build/EditorPerf.xml
+++ b/Engine/Build/EditorPerf.xml
@@ -244,7 +244,7 @@
 			<Regex Pattern="(.*)\{(.*)\}" Capture="Project;ProjectTarget" Input="$(Project)" Optional="true"/>
 
 			<!-- Project Map Args-->
-			<Property Name="ProjectCombinedArgs" Value="$($(Project)SpecificArgs) $(Project)MapArgs) $(AdditionalRuntimeArgs)"/>
+			<Property Name="ProjectCombinedArgs" Value="$($(Project)SpecificArgs) $($(Project)MapArgs) $(AdditionalRuntimeArgs)"/>
 
 			<!-- Use as base string for directories and key for node names -->
 			<Property Name="OutputName" Value="$(Project)$(Platform)EditorPerf"/>
