Add an option to disable the summarize trace step

diff --git a/Engine/Build/EditorPerf.xml b/Engine/Build/EditorPerf.xml
--- a/Engine/Build/EditorPerf.xml
+++ b/Engine/Build/EditorPerf.xml
@@ -162,6 +162,9 @@
 	<Property Name="ProfileArgs" Value="$(ProfileArgs) -NoLoadLevelAtStartup -Benchmarking -NODDCCLEANUP"/>
 	<Property Name="ProfileArgs" Value="$(ProfileArgs) -NoShaderDistrib -VerboseShaderLogging"/>
 
+	<!-- options for summarize runs -->
+	<Option Name="RunTraceSummary" DefaultValue="true" Description="Summarize the insights trace after the run"/>
+
 	<!-- Gauntlet arguments for commandlet, summarize runs -->
 	<Property Name="CommandletArgs" Value="-notraceserver"/>
 	<Property Name="SummarizeArgs" Value=""/>
@@ -169,8 +172,8 @@
 
 	<!-- Publish telemetry options -->
 	<Option Name="TelemetryConfig" DefaultValue="" Description="The configuration to use when publishing telemetry data to a database"/>
-	<Property Name="TelemetryConfig" Value="UETelemetryStaging" If="$(IsBuildMachine) And $(IsPreflight) And '$(TelemetryConfig)' == ''"/>
-	<Property Name="TelemetryConfig" Value="UETelemetry" If="$(IsBuildMachine) And '$(TelemetryConfig)' == ''"/>
+	<Property Name="TelemetryConfig" Value="UETelemetryStaging" If="$(IsBuildMachine) And $(IsPreflight) And '$(TelemetryConfig)' == '' And $(RunTraceSummary)"/>
+	<Property Name="TelemetryConfig" Value="UETelemetry" If="$(IsBuildMachine) And '$(TelemetryConfig)' == '' And $(RunTraceSummary)"/>
 
 	<!-- The main test running macro logic that is expanded with platform and dependency node names -->
 	<Macro Name="RunTest" Arguments="TestNameSuffix;Platform;ProjectName;ProjectTarget;NodeName;Phase;Iteration;Args;EditorArgs">
@@ -198,7 +201,9 @@
 		<Command Name="RunUnreal" Arguments="$(AllGauntletArgs) $(ProfileArgs) $(RunArgs) -sessionlabel=&quot;$(TestName)&quot; -tracefile=&quot;$(TraceFilePath)&quot;"/>
 
 		<!-- Translate the utrace file into csv data -->
-		<Commandlet Project="$(ProjectTarget)" Name="SummarizeTrace" Arguments="-inputfile=&quot;$(TraceFilePath)&quot; -testname=&quot;$(TestName)&quot; $(SummarizeArgs) $(CommandletArgs)"/>
+		<Do If="$(RunTraceSummary)">
+			<Commandlet Project="$(ProjectTarget)" Name="SummarizeTrace" Arguments="-inputfile=&quot;$(TraceFilePath)&quot; -testname=&quot;$(TestName)&quot; $(SummarizeArgs) $(CommandletArgs)"/>
+		</Do>
 
 		<!-- Save all files named for this run in the trace directory -->
 		<Copy From="$(TraceDir)/$(IterationName)*.*" To="$(TestOutputDir)"/>
