Disable the summarize trace by default

diff --git a/Engine/Build/EditorPerf.xml b/Engine/Build/EditorPerf.xml
--- a/Engine/Build/EditorPerf.xml
+++ b/Engine/Build/EditorPerf.xml
@@ -163,7 +163,7 @@
 	<Property Name="ProfileArgs" Value="$(ProfileArgs) -NoShaderDistrib -VerboseShaderLogging"/>
 
 	<!-- options for summarize runs -->
-	<Option Name="RunTraceSummary" DefaultValue="true" Description="Summarize the insights trace after the run"/>
+	<Option Name="RunTraceSummary" DefaultValue="false" Description="Summarize the insights trace after the run"/>
 
 	<!-- Gauntlet arguments for commandlet, summarize runs -->
 	<Property Name="CommandletArgs" Value="-notraceserver"/>
