add Engine/Build/Mac/Resources to installed build

diff --git a/Engine/Build/InstalledEngineFilters.xml b/Engine/Build/InstalledEngineFilters.xml
--- a/Engine/Build/InstalledEngineFilters.xml
+++ b/Engine/Build/InstalledEngineFilters.xml
@@ -340,6 +340,9 @@
 		Engine/Build/BatchFiles/RunUAT.command
 		Engine/Build/BatchFiles/RunUAT.sh
 		Engine/Build/BatchFiles/RunUBT.sh
+        
+        <!-- Mac build resources (plist entitlement etc.) -->
+        Engine/Build/Mac/Resources/...
 
 		<!-- Mono -->
 		Engine/Binaries/ThirdParty/Mono/Mac/...
