Add using Win64Compile Agent for Compile instead of just Win64
#rnx

diff --git a/Engine/Build/EditorPerf.xml b/Engine/Build/EditorPerf.xml
--- a/Engine/Build/EditorPerf.xml
+++ b/Engine/Build/EditorPerf.xml
@@ -22,7 +22,7 @@
 
 	<!-- Platform Options -->
 	<Option Name="Platforms" DefaultValue="Win64" Description="The list of platforms to profile"/>
-	<Property Name="Win64CompileAgentType" Value="Win64"/>
+	<Property Name="Win64CompileAgentType" Value="Win64Compile;Win64"/>
 	<Property Name="MacCompileAgentType" Value="TestMac"/>
 	<Property Name="LinuxCompileAgentType" Value="Linux"/>
 	<Error Message="No platforms specified" If="'$(Platforms)' == ''"/>
