NDA platforms in restricted - move some parts into the platform folders

diff --git a/Engine/Build/LowLevelTests.xml b/Engine/Build/LowLevelTests.xml
--- a/Engine/Build/LowLevelTests.xml
+++ b/Engine/Build/LowLevelTests.xml
@@ -13,11 +13,38 @@
 
 	<Property Name="AlternateDevicePoolPlatformsProperty" Value=""/>
 
+	<!-- Platform specific macros -->
+	<Macro Name="Catch2BuildAppendExtraCMakeArgsPlatform" Arguments="Platform"/>
+	<Macro Name="Catch2BuildAppendExtraCMakeArgs" Arguments="ExtraArgValue">
+		<Property Name="ExtraCMakeArgs" Value="-CMakeAdditionalArguments=&quot;$(CommonExtraCmakeArgs)$(ExtraArgValue)&quot;"/>
+	</Macro>
+	<Macro Name="SelectAgentType" Arguments="Platform;BuildPlatform"/>
+
+	<Macro Name="SetupPlatformProperties" Arguments="Platform;BuildPlatform" OptionalArguments="AlternateDevicePool">
+		<Property Name="AllTestablePlatforms" Value="$(AllTestablePlatforms);$(Platform)"/>
+		<Do If="'$(BuildPlatform)' == 'Win64'">
+			<Property Name="PlatformsBuiltFromWin64" Value="$(PlatformsBuiltFromWin64);$(Platform)"/>
+		</Do>
+		<Do If="'$(BuildPlatform)' == 'Mac'">
+			<Property Name="PlatformsBuiltFromMac" Value="$(PlatformsBuiltFromMac);$(Platform)"/>
+		</Do>
+		<Do If="'$(BuildPlatform)' == 'Linux'">
+			<Property Name="PlatformsBuiltFromLinux" Value="$(PlatformsBuiltFromLinux);$(Platform)"/>
+		</Do>
+		<Do If="'$(AlternateDevicePool)' != ''">
+			<Property Name="AlternateDevicePoolPlatformsProperty" Value="$(AlternateDevicePoolPlatformsProperty);$(AlternateDevicePool)"/>
+		</Do>
+	</Macro>
+	<Macro Name="SetupAllPlatformProperties"/>
+
+	<!-- All test metadata scripts along with platform specific versions for both public and internal paths -->
 	<Include Script="LowLevelTests/*.xml"/>
 	<Include Script="../Platforms/*/Build/LowLevelTests/*.xml"/>
 	<Include Script="../Restricted/NotForLicensees/Build/LowLevelTests/*.xml"/>
 	<Include Script="../Restricted/NotForLicensees/Platforms/*/Build/LowLevelTests/*.xml"/>
 
+	<Expand Name="SetupAllPlatformProperties"/>
+
 	<Option Name="BuildCatch2Libs" DefaultValue="false" Description="True if building Catch2 static libraries"/>
 	<Option Name="TestCatch2Libs" DefaultValue="false" Description="True if additionally building Catch2 with extras and tests"/>
 	<Option Name="Catch2LibVariation" DefaultValue="VS2019" Description="Catch2 library variation, e.g. VS2022"/>
@@ -114,9 +141,7 @@
 	<ForEach Name="BuildPlatform" Values="$(BuildPlatforms)">
 		<ForEach Name="Platform" Values="$(PlatformsBuiltFrom$(BuildPlatform))">
 			<Property Name="AgentType" Value="$(BuildPlatform)"/>
-			<Do If="'$($(Platform)RequiresSpecialAgent)' == 'True'">
-				<Expand Name="SelectAgentType" Platform="$(Platform)" BuildPlatform="$(BuildPlatform)" />
-			</Do>
+			<Expand Name="SelectAgentType" Platform="$(Platform)" BuildPlatform="$(BuildPlatform)" />
 			<Property Name="BuildAgentType" Value="$(AgentType)"/>
 			<Do If="'$(Platform)' == 'Linux'">
 				<Property Name="BuildAgentType" Value="Win64"/>
@@ -142,9 +167,7 @@
 						</Do>
 						<Property Name="ExtraCMakeArgs" Value="-TargetArchitecture=x64 -CMakeAdditionalArguments=&quot;$(CommonExtraCmakeArgs)-DSET_CONFIG_NO_COLOUR_WIN32=OFF&quot;"/>
 					</Do>
-					<Do If="'$($(Platform)RequiresCatch2BuildExtraCMakeArgs)' == 'True'">
-						<Expand Name="Catch2BuildAppendExtraCMakeArgs" Platform="$(Platform)" />
-					</Do>
+					<Expand Name="Catch2BuildAppendExtraCMakeArgsPlatform" Platform="$(Platform)"/>
 					<Do If="'$(BuildPlatform)' == 'Mac'">
 						<Property Name="Generator" Value="Xcode"/>
 						<Property Name="TargetConfigs" Value="Debug+Release"/>
diff --git a/Engine/Build/LowLevelTests/HeadlessChaos.xml b/Engine/Build/LowLevelTests/HeadlessChaos.xml
--- a/Engine/Build/LowLevelTests/HeadlessChaos.xml
+++ b/Engine/Build/LowLevelTests/HeadlessChaos.xml
@@ -6,7 +6,7 @@
   <Property Name="HeadlessChaosShort" Value="Headless Chaos" />
   <Property Name="HeadlessChaosTarget" Value="HeadlessChaos" />
   <Property Name="HeadlessChaosBinariesRelative" Value="Engine\Binaries" />
-  <Property Name="HeadlessChaosReportType" />
+  <Property Name="HeadlessChaosReportType" Value="console" />
   <Property Name="HeadlessChaosInitialExtraArgs" />
   <Property Name="HeadlessChaosHasAfterSteps" Value="False" />
   <Property Name="HeadlessChaosUsesCatch2" Value="False" />
diff --git a/Engine/Build/LowLevelTests/NetworkPredictionPlugin.xml b/Engine/Build/LowLevelTests/NetworkPredictionPlugin.xml
--- a/Engine/Build/LowLevelTests/NetworkPredictionPlugin.xml
+++ b/Engine/Build/LowLevelTests/NetworkPredictionPlugin.xml
@@ -6,7 +6,7 @@
   <Property Name="NetworkPredictionPluginShort" Value="Net Prediction" />
   <Property Name="NetworkPredictionPluginTarget" Value="NetworkPredictionTests" />
   <Property Name="NetworkPredictionPluginBinariesRelative" Value="Engine\Binaries" />
-  <Property Name="NetworkPredictionPluginReportType" />
+  <Property Name="NetworkPredictionPluginReportType" Value="console" />
   <Property Name="NetworkPredictionPluginInitialExtraArgs" />
   <Property Name="NetworkPredictionPluginHasAfterSteps" Value="False" />
   <Property Name="NetworkPredictionPluginUsesCatch2" Value="True" />
diff --git a/Engine/Build/LowLevelTests/ReplicationSystem.xml b/Engine/Build/LowLevelTests/ReplicationSystem.xml
--- a/Engine/Build/LowLevelTests/ReplicationSystem.xml
+++ b/Engine/Build/LowLevelTests/ReplicationSystem.xml
@@ -6,7 +6,7 @@
   <Property Name="ReplicationSystemShort" Value="Replication System" />
   <Property Name="ReplicationSystemTarget" Value="ReplicationSystemLowLevelTests" />
   <Property Name="ReplicationSystemBinariesRelative" Value="Engine\Binaries" />
-  <Property Name="ReplicationSystemReportType" />
+  <Property Name="ReplicationSystemReportType" Value="console" />
   <Property Name="ReplicationSystemInitialExtraArgs" />
   <Property Name="ReplicationSystemHasAfterSteps" Value="False" />
   <Property Name="ReplicationSystemUsesCatch2" Value="True" />
diff --git a/Engine/Source/Programs/AutomationTool/LowLevelTests/RunLowLevelTests.cs b/Engine/Source/Programs/AutomationTool/LowLevelTests/RunLowLevelTests.cs
--- a/Engine/Source/Programs/AutomationTool/LowLevelTests/RunLowLevelTests.cs
+++ b/Engine/Source/Programs/AutomationTool/LowLevelTests/RunLowLevelTests.cs
@@ -608,7 +608,7 @@ namespace LowLevelTests
 
 		public string CopyDeviceReportTo(IAppInstall InAppInstall, UnrealTargetPlatform InPlatform, string InTestApp, string InBuildPath, string InTargetDirectory)
 		{
-			string ReportRelativePath = GetTargetReportPath(InPlatform, InTestApp, InBuildPath);
+			string ReportRelativePath = Path.GetFileName(GetTargetReportPath(InPlatform, InTestApp, InBuildPath));
 			string ReportPath = Path.Combine(InBuildPath, ReportRelativePath);
 			string ExpectedLocalPath = Path.Combine(InTargetDirectory, ReportRelativePath);
 			if (!ExpectedLocalPath.Equals(ReportPath))
diff --git a/Engine/Source/Programs/UnrealBuildTool/Configuration/TestModuleRules.cs b/Engine/Source/Programs/UnrealBuildTool/Configuration/TestModuleRules.cs
--- a/Engine/Source/Programs/UnrealBuildTool/Configuration/TestModuleRules.cs
+++ b/Engine/Source/Programs/UnrealBuildTool/Configuration/TestModuleRules.cs
@@ -417,10 +417,14 @@ namespace UnrealBuildTool
 			/// </summary>
 			public string TestShortName { get; set; }
 
+			private string ReportTypePrivate = "console";
 			/// <summary>
 			/// Type of Catch2 report, defaults to console.
 			/// </summary>
-			public string ReportType { get; set; }
+			public string ReportType {
+				get => ReportTypePrivate;
+				set => ReportTypePrivate = value;
+			}
 
 			/// <summary>
 			/// Is this test disabled?
