add modern Xcode related scripts to installed build

diff --git a/Engine/Build/InstalledEngineFilters.xml b/Engine/Build/InstalledEngineFilters.xml
--- a/Engine/Build/InstalledEngineFilters.xml
+++ b/Engine/Build/InstalledEngineFilters.xml
@@ -337,6 +337,8 @@
 		Engine/Build/BatchFiles/Mac/RunMono.sh
 		Engine/Build/BatchFiles/Mac/RunXBuild.sh
 		Engine/Build/BatchFiles/Mac/XcodeBuild.sh
+        Engine/Build/BatchFiles/Mac/GenerateUniversalDSYM.sh
+        Engine/Build/BatchFiles/Mac/UpdateVersionAfterBuild.sh
 		Engine/Build/BatchFiles/RunUAT.command
 		Engine/Build/BatchFiles/RunUAT.sh
 		Engine/Build/BatchFiles/RunUBT.sh
