Separate the traces data copy from the running of the tests so they can still be copied on failures

diff --git a/Engine/Build/EditorPerf.xml b/Engine/Build/EditorPerf.xml
--- a/Engine/Build/EditorPerf.xml
+++ b/Engine/Build/EditorPerf.xml
@@ -164,7 +164,6 @@
 
 	<!-- The main test running macro logic that is expanded with platform and dependency node names -->
 	<Macro Name="RunTest" Arguments="TestNameSuffix;Platform;ProjectName;ProjectTarget;NodeName;Phase;Iteration;Args;EditorArgs">
-
 		<Property Name="TestName" Value="EditorPerf$(Phase)$(TestNameSuffix)"/>
 		<Property Name="IterationSuffix" Value=""/>
 		<Property Name="IterationSuffix" Value="$(Iteration)" If="'$(Iteration)' != '1'"/>
@@ -175,6 +174,7 @@
 		<Property Name="TraceFilePath" Value="$(TraceDir)/$(IterationName).utrace"/>
 		<Property Name="TelemetryFileName" Value="$(IterationName)Telemetry.csv"/>
 		<Property Name="TelemetryFilePath" Value="$(TraceDir)/$(TelemetryFileName)"/>
+
 		<Property Name="LogOutputParams" Value="-logdir=&quot;$(TestOutputDir)&quot;"/>
 		<Property Name="ReportParams" Value="-WriteTestResultsForHorde -HordeTestDataKey=&quot;$(NodeName)$(PreflightSuffix)&quot;"/>
 		<Property Name="ArtifactName" Value="$(Phase)$(IterationSuffix)$(TestNameSuffix)"/>
@@ -191,15 +191,17 @@
 		<Do If="$(RunTraceSummary)">
 			<Commandlet Project="$(ProjectTarget)" Name="SummarizeTrace" Arguments="-inputfile=&quot;$(TraceFilePath)&quot; -testname=&quot;$(TestName)&quot; $(SummarizeArgs) $(CommandletArgs)"/>
 		</Do>
-
-		<!-- Save all files named for this run in the trace directory -->
-		<Copy From="$(TraceDir)/$(IterationName)*.*" To="$(TestOutputDir)"/>
-
+		
 		<Do If="'$(TelemetryConfig)' != ''">
 			<!-- Publish telemetry to the db -->
 			<Command Name="PublishUnrealAutomationTelemetry" Arguments="-CSVFile=&quot;$(TelemetryFilePath)&quot; -project=$(Project) -changelist=$(Change) -branch=$(Branch) -TelemetryConfig=$(TelemetryConfig)"/>
-		</Do>
-
+		</Do>		
+	</Macro>
+	
+	<Macro Name="UploadTestTraces" Arguments="Platform;ProjectName">		
+		<!-- Save all files named for this project in the trace directory -->
+		<Property Name="TestOutputDir" Value="$(AutomationOutputDir)/$(ProjectName)$(Platform)EditorPerf"/>
+		<Copy From="$(TraceDir)/$(ProjectName)$(Platform)EditorPerf*.*" To="$(TestOutputDir)"/>
 	</Macro>
 
 	<!-- This will be set to the names of all nodes we want to execute based on specified platforms/options -->
@@ -258,6 +260,7 @@
 			<!-- Run the Tests -->
 			<Property Name="RunNodeRequires" Value="$(CompileNodeName)"/>
 			<Property Name="RunNodeList" Value=""/> <!-- Each node in this the nested loops below, they need to depend on predecessors -->
+			<Property Name="UploadNodeList" Value=""/>
 			<Agent Name="EditorPerf Agent $(Platform) $(Project)" Type="EditorPerf$(Platform)">
 				<ForEach Name="Phase" Values="$(Phases)">
 					<Do If="$(EnableBasePhaseTests)">
@@ -312,9 +315,14 @@
 						<Property Name="RunNodeList" Value="$(RunNodeList);$(RunNodeLazyLoadZenName)" If="$(EnableLazyLoadZenTest)"/>
 					</Do>
 				</ForEach>
+				<Property Name="RunNodeUploadPhaseTraces" Value="Upload $(OutputName) Traces"/>
+				<Node Name="$(RunNodeUploadPhaseTraces)">
+					<Expand Name="UploadTestTraces" Platform="$(Platform)" ProjectName="$(Project)"/>
+				</Node>
+				<Property Name="UploadNodeList" Value="$(RunNodeUploadPhaseTraces)"/>
 			</Agent>
-			<Label Category="$(RunCategory)" Name="$(Project) Editor Run" Requires="$(RunNodeRequires);$(RunNodeList)" UgsBadge="$(Project) Editor Run"/>
-			<Property Name="RunGraphNodes" Value="$(RunGraphNodes);$(RunNodeList)"/>
+			<Label Category="$(RunCategory)" Name="$(Project) Editor Run" Requires="$(RunNodeRequires);$(RunNodeList);$(UploadNodeList)" UgsBadge="$(Project) Editor Run"/>
+			<Property Name="RunGraphNodes" Value="$(RunGraphNodes);$(RunNodeList);$(UploadNodeList)"/>
 
 			<!-- Baseline the Results -->
 			<Property Name="BaselineNodeName" Value="Baseline $(OutputName)"/>
