Update LowLevelTests to use BuildCookRun when test specifies StagesWithProjectFile.
 - LowLevelTest.xml horde build step as a separate label
 - Relevant <General>.xml files updated to support the new build step
 - Added UpdateBuildGraphPropertiesFile.Automatoin.cs to help find files to checkout
 - OnlineTests updated to utilize new BuildCookRun param

diff --git a/Engine/Build/LowLevelTests.xml b/Engine/Build/LowLevelTests.xml
--- a/Engine/Build/LowLevelTests.xml
+++ b/Engine/Build/LowLevelTests.xml
@@ -110,7 +110,20 @@
 		<Property Name="AfterTestNodes" Value="$(AfterTestNodes);$(TestName) AfterStep $(Platform)" If="$(RunAllTests) Or $(Run$(TestName)Tests)"/>
 	</Macro>
 
+	<Macro Name="BuildTest" Arguments="Platform;TestName">
+		<Property Name="LocalAchiveDirectory" Value="$(RootDir)/$($(TestName)Target)"/>
+		<Do If="'$($(TestName)StagesWithProjectFile)' == 'True'">
+			<Command Name="BuildCookRun" Arguments="-Build -SkipCook -Stage -NoSubmit -Archive -ArchiveDirectory=$(LocalAchiveDirectory) -project=$($(TestName)Target) -configuration=$(Configuration) -platform=$(Platform) $($(Platform)$(TestName)ExtraCompilationArgs)"/>
+		</Do>
+		<Do If="'$($(TestName)StagesWithProjectFile)' != 'True'">
+			<Compile Target="$($(TestName)Target)" Configuration="$(Configuration)" Platform="$(Platform)" AllowXGE="false" Arguments="$($(Platform)$(TestName)ExtraCompilationArgs)"/>
+		</Do>
+	</Macro>
+
 	<Macro Name="DeployAndTest" Arguments="TestName;Platform">
+		<Property Name="PlatformStageDirName" Value="$(Platform)"/>
+		<Property Name="PlatformStageDirName" Value="Windows" If="'$(Platform)' == 'Win64'"/>
+
 		<Property Name="DevicePool" Value="$(GauntletDevicePool)" />
 		<Property Name="DevicePool" Value="$(AlternateGauntletDevicePool)" If="ContainsItem('$(AlternateDevicePoolPlatforms)', '$(Platform)', ';')" />
 		<Property Name="TestArgsExt" Value="$(TestArgs) $($(TestName)InitialExtraArgs) -deviceurl=$(GauntletDeviceService) -devicepool=$(DevicePool)"/>
@@ -121,6 +134,9 @@
 		<Property Name="TestArgsExt" Value="$(TestArgsExt) -tags=&quot;$(TestTags)&quot;"/>
 		<Property Name="TestArgsExt" Value="$(TestArgsExt) --buildmachine" If="$(IsBuildMachine)"/>
 		<Property Name="TestArgsExt" Value="$(TestArgsExt) -timeout=$(PerTestTimeout)"/>
+		<Property Name="TestArgsExt" Value="$(TestArgsExt) -SkipStage -Build=&quot;$(RootDir)/OnlineTests/$(PlatformStageDirName)&quot; $($(Platform)ExtraArgs)" If="'$($(TestName)StagesWithProjectFile)' == 'True'"/>
+		<Property Name="TestArgsExt" Value="$(TestArgsExt) -build=&quot;$($(TestName)BinariesRelative)\$(Platform)\$($(TestName)Target)$(TestModeSuffix)&quot;" If="'$($(TestName)StagesWithProjectFile)' != 'True'"/>
+		
 		<Property Name="TestArgsExt" Value="$(TestArgsExt) $($(Platform)ExtraArgs)"/>
 		<Do If="'$(Device)' != ''">
 			<Property Name="TestArgsExt" Value="$(TestArgsExt) -device=$(Device) "/>
@@ -135,7 +151,7 @@
 			<Property Name="TestArgsExt" Value=""/>
 		</Do>
 		<Command Name="RunLowLevelTests"
-				Arguments="-test=LowLevelTests -testapp=$($(TestName)Target)$(TestModeSuffix) -platform=$(Platform) -configuration=$(Configuration) $(TestArgsExt) -build=&quot;$($(TestName)BinariesRelative)\$(Platform)\$($(TestName)Target)$(TestModeSuffix)&quot; -VeryVerbose"/>
+				Arguments="-test=LowLevelTests -testapp=$($(TestName)Target)$(TestModeSuffix) -platform=$(Platform) -configuration=$(Configuration) $(TestArgsExt) -VeryVerbose"/>
 	</Macro>
 
 	<ForEach Name="BuildPlatform" Values="$(BuildPlatforms)">
@@ -198,8 +214,12 @@
 							<Property Name="TestRequirements" Value="$(TestRequirements);Catch2 Build Library $(Platform)"/>
 						</Do>
 						
-						<Node Name="$($(TestName)Short) Tests $(Platform)" Requires="$(TestRequirements)">
-							<Compile Target="$($(TestName)Target)" Configuration="$(Configuration)" Platform="$(Platform)" AllowXGE="false" Arguments="$($(Platform)$(TestName)ExtraCompilationArgs)"/>
+						<Node Name="$($(TestName)Short) Build Tests $(Platform)" Requires="$(TestRequirements)">
+							<Expand Name="BuildTest" TestName="$(TestName)" Platform="$(Platform)"/>
+						</Node>
+						<Label Category="Build" Name="$($(TestName)Short) Build $(Platform)" Requires="$($(TestName)Short) Build Tests $(Platform)"/>
+
+						<Node Name="$($(TestName)Short) Tests $(Platform)" Requires="$(TestRequirements);$($(TestName)Short) Build Tests $(Platform)">
 							<Do If="'$($(Platform)$(TestName)RunSupported)' == 'True'">
 								<Expand Name="DeployAndTest" TestName="$(TestName)" Platform="$(Platform)"/>
 							</Do>
diff --git a/Engine/Build/LowLevelTests/Foundation.xml b/Engine/Build/LowLevelTests/Foundation.xml
--- a/Engine/Build/LowLevelTests/Foundation.xml
+++ b/Engine/Build/LowLevelTests/Foundation.xml
@@ -4,6 +4,7 @@
   <Property Name="FoundationUsesCatch2" Value="True" />
   <Property Name="FoundationDisabled" Value="False" />
   <Property Name="FoundationShort" Value="Foundation" />
+  <Property Name="FoundationStagesWithProjectFile" Value="False" />
   <Property Name="FoundationTarget" Value="FoundationTests" />
   <Property Name="FoundationBinariesRelative" Value="Engine\Binaries" />
   <Property Name="FoundationReportType" Value="xml" />
diff --git a/Engine/Build/LowLevelTests/HeadlessChaos.xml b/Engine/Build/LowLevelTests/HeadlessChaos.xml
--- a/Engine/Build/LowLevelTests/HeadlessChaos.xml
+++ b/Engine/Build/LowLevelTests/HeadlessChaos.xml
@@ -4,6 +4,7 @@
   <Property Name="HeadlessChaosUsesCatch2" Value="False" />
   <Property Name="HeadlessChaosDisabled" Value="False" />
   <Property Name="HeadlessChaosShort" Value="Headless Chaos" />
+  <Property Name="HeadlessChaosStagesWithProjectFile" Value="False" />
   <Property Name="HeadlessChaosTarget" Value="HeadlessChaos" />
   <Property Name="HeadlessChaosBinariesRelative" Value="Engine\Binaries" />
   <Property Name="HeadlessChaosReportType" Value="console" />
diff --git a/Engine/Build/LowLevelTests/NetworkPredictionPlugin.xml b/Engine/Build/LowLevelTests/NetworkPredictionPlugin.xml
--- a/Engine/Build/LowLevelTests/NetworkPredictionPlugin.xml
+++ b/Engine/Build/LowLevelTests/NetworkPredictionPlugin.xml
@@ -4,6 +4,7 @@
   <Property Name="NetworkPredictionPluginUsesCatch2" Value="True" />
   <Property Name="NetworkPredictionPluginDisabled" Value="False" />
   <Property Name="NetworkPredictionPluginShort" Value="Net Prediction" />
+  <Property Name="NetworkPredictionPluginStagesWithProjectFile" Value="False" />
   <Property Name="NetworkPredictionPluginTarget" Value="NetworkPredictionTests" />
   <Property Name="NetworkPredictionPluginBinariesRelative" Value="Engine\Binaries" />
   <Property Name="NetworkPredictionPluginReportType" Value="console" />
diff --git a/Engine/Build/LowLevelTests/ReplicationSystem.xml b/Engine/Build/LowLevelTests/ReplicationSystem.xml
--- a/Engine/Build/LowLevelTests/ReplicationSystem.xml
+++ b/Engine/Build/LowLevelTests/ReplicationSystem.xml
@@ -4,6 +4,7 @@
   <Property Name="ReplicationSystemUsesCatch2" Value="True" />
   <Property Name="ReplicationSystemDisabled" Value="False" />
   <Property Name="ReplicationSystemShort" Value="Replication System" />
+  <Property Name="ReplicationSystemStagesWithProjectFile" Value="False" />
   <Property Name="ReplicationSystemTarget" Value="ReplicationSystemLowLevelTests" />
   <Property Name="ReplicationSystemBinariesRelative" Value="Engine\Binaries" />
   <Property Name="ReplicationSystemReportType" Value="console" />
diff --git a/Engine/Source/Programs/AutomationTool/LowLevelTests/RunLowLevelTests.cs b/Engine/Source/Programs/AutomationTool/LowLevelTests/RunLowLevelTests.cs
--- a/Engine/Source/Programs/AutomationTool/LowLevelTests/RunLowLevelTests.cs
+++ b/Engine/Source/Programs/AutomationTool/LowLevelTests/RunLowLevelTests.cs
@@ -497,7 +497,9 @@ namespace LowLevelTests
 				ParentDirPath = Directory.GetParent(Executable).FullName;
 				BuildExecutableName = Path.GetFileNameWithoutExtension(Executable);
 
-				if (!ParentDirPath.Contains(InPlatform.ToString(), StringComparison.OrdinalIgnoreCase))
+				bool ContainsPlatformName = ParentDirPath.Contains(InPlatform.ToString(), StringComparison.OrdinalIgnoreCase);
+				bool ContainsWindowsName = (InPlatform == UnrealTargetPlatform.Win64 && ParentDirPath.Contains("Windows", StringComparison.OrdinalIgnoreCase));
+				if ( !(ContainsPlatformName || ContainsWindowsName) )
 				{
 					Log.Error("ParentPath did not have platform ({1}) : {0}", ParentDirPath, InPlatform);
 					continue;
diff --git a/Engine/Source/Programs/AutomationTool/LowLevelTests/UpdateGraphsLowLevelTests.Automation.cs b/Engine/Source/Programs/AutomationTool/LowLevelTests/UpdateGraphsLowLevelTests.Automation.cs
new file mode 100644
--- /dev/null
+++ b/Engine/Source/Programs/AutomationTool/LowLevelTests/UpdateGraphsLowLevelTests.Automation.cs
@@ -0,0 +1,66 @@
+// Copyright Epic Games, Inc. All Rights Reserved.
+
+using System.Collections.Generic;
+using AutomationTool;
+
+using UnrealBuildBase;
+using Gauntlet;
+using Microsoft.Extensions.FileSystemGlobbing.Abstractions;
+using Microsoft.Extensions.FileSystemGlobbing;
+using System.IO;
+using EpicGames.Core;
+using Microsoft.Extensions.Logging;
+using System.Linq;
+
+namespace LowLevelTests
+{
+	using Log = Gauntlet.Log;
+	using LogLevel = Gauntlet.LogLevel;
+
+	[RequireP4]
+	[Help("Perforce helper to checkout all related LowLevelTest Build.xml files into perforce.")]
+	[Help("EditCl=<CL>","(optional) the CL we wish to append graph changes to. Defaults to 'default'")]
+	public class UpdateGraphsLowLevelTests : BuildCommand
+	{
+		public override ExitCode Execute()
+		{
+			Log.Level = LogLevel.VeryVerbose;
+			Globals.Params = new Params(Params);
+
+			int EditCL = ParseParamNullableInt("EditCL") ?? 0;
+
+			List<string> BuildPaths = new List<string>() { "Build/LowLevelTests/*.xml"};
+			Matcher matcher = new();
+			matcher.AddIncludePatterns(BuildPaths);
+
+
+			List<string> LowLevelTestXml = new List<string>();
+			foreach(DirectoryReference EngineBase in GetAllEngineDirectories())
+			{
+				foreach (DirectoryReference ExtendedEngineDir in Unreal.GetExtensionDirs(EngineBase))
+				{
+					Logger.LogInformation(ExtendedEngineDir.FullName);
+					IEnumerable<string> matchingFiles = matcher.GetResultsInFullPath(ExtendedEngineDir.FullName);
+					foreach (string v in matchingFiles)
+					{
+						Logger.LogInformation(v);
+						LowLevelTestXml.Add(v);
+					}
+				}
+			}
+
+			if(EditCL == 0)
+			{
+				P4.BatchedCommand("edit", LowLevelTestXml);
+			}
+			else
+			{
+				P4.BatchedCommand($"edit -c {EditCL}", LowLevelTestXml);
+			}
+
+			//TODO(SMA): Enable bUpdateBuildGraphPropertiesFile in config and run UpdateBuildGraphPropertiesFile on all relevant .Build.cs files.
+
+			return ExitCode.Success;
+		}
+	}
+}
diff --git a/Engine/Source/Programs/UnrealBuildTool/Configuration/TestModuleRules.cs b/Engine/Source/Programs/UnrealBuildTool/Configuration/TestModuleRules.cs
--- a/Engine/Source/Programs/UnrealBuildTool/Configuration/TestModuleRules.cs
+++ b/Engine/Source/Programs/UnrealBuildTool/Configuration/TestModuleRules.cs
@@ -228,14 +228,23 @@ namespace UnrealBuildTool
 
 					InsertOrUpdateTestFlagProperty(ref lastUpdatedNode, TestMetadata.TestName, "Disabled", Convert.ToString(TestMetadata.Disabled));
 					InsertOrUpdateTestFlagProperty(ref lastUpdatedNode, TestMetadata.TestName, "Short", Convert.ToString(TestMetadata.TestShortName));
+					InsertOrUpdateTestFlagProperty(ref lastUpdatedNode, TestMetadata.TestName, "StagesWithProjectFile", Convert.ToString(TestMetadata.StagesWithProjectFile));
 					InsertOrUpdateTestFlagProperty(ref lastUpdatedNode, TestMetadata.TestName, "Target", Convert.ToString(TestTargetName));
 					InsertOrUpdateTestFlagProperty(ref lastUpdatedNode, TestMetadata.TestName, "BinariesRelative", Convert.ToString(TestBinariesPath));
 					InsertOrUpdateTestFlagProperty(ref lastUpdatedNode, TestMetadata.TestName, "ReportType", Convert.ToString(TestMetadata.ReportType));
-					InsertOrUpdateTestFlagProperty(ref lastUpdatedNode, TestMetadata.TestName, "InitialExtraArgs", Convert.ToString(TestMetadata.InitialExtraArgs));
+
+					string ExtraArguments = Convert.ToString(TestMetadata.InitialExtraArgs);
+					if(TestMetadata.StagesWithProjectFile)
+					{
+						//  -BaseFromWorkingDir Is needed here for Programs to properly detect .ini files from its irregular platform path.
+						ExtraArguments += $"--extra-args -BaseFromWorkingDir";
+					}
+
+					InsertOrUpdateTestFlagProperty(ref lastUpdatedNode, TestMetadata.TestName, "InitialExtraArgs", ExtraArguments);
 					InsertOrUpdateTestFlagProperty(ref lastUpdatedNode, TestMetadata.TestName, "HasAfterSteps", Convert.ToString(TestMetadata.HasAfterSteps));
 					InsertOrUpdateTestFlagProperty(ref lastUpdatedNode, TestMetadata.TestName, "UsesCatch2", Convert.ToString(TestMetadata.UsesCatch2));
 
-					InsertOrUpdateTestOption(ref lastUpdatedNode, TestMetadata.TestName, "Run" + TestMetadata.TestShortName + "Tests", "Run", "Tests", false.ToString());
+					InsertOrUpdateTestOption(ref lastUpdatedNode, TestMetadata.TestName, $"Run {TestMetadata.TestShortName} Tests", "Run", "Tests", false.ToString());
 
 					List<UnrealTargetPlatform> AllSupportedPlatforms = new List<UnrealTargetPlatform>();
 					object[] SupportedPlatforms = GetType().GetCustomAttributes(typeof(SupportedPlatformsAttribute), false);
@@ -437,6 +446,13 @@ namespace UnrealBuildTool
 				set => ReportTypePrivate = value;
 			}
 
+			/// <summary>
+			/// Does this test use project files for staging additional files?
+			/// This will append `-SkipStage -Build=$(RootDir)/$($(TestName)Target)` to InitialExtraArgs
+			/// and cause the build to use BuildCookRun instead of a Compile step
+			/// </summary>
+			public bool StagesWithProjectFile { get; set; }
+
 			/// <summary>
 			/// Is this test disabled?
 			/// </summary>
