Back out of changes to LowLevelTests.xml Submitted via 26194455

diff --git a/Engine/Build/LowLevelTests.xml b/Engine/Build/LowLevelTests.xml
--- a/Engine/Build/LowLevelTests.xml
+++ b/Engine/Build/LowLevelTests.xml
@@ -110,20 +110,7 @@
 		<Property Name="AfterTestNodes" Value="$(AfterTestNodes);$(TestName) AfterStep $(Platform)" If="$(RunAllTests) Or $(Run$(TestName)Tests)"/>
 	</Macro>
 
-	<Macro Name="BuildTest" Arguments="Platform;TestName">
-		<Property Name="LocalAchiveDirectory" Value="$(RootDir)/$($(TestName)Target)"/>
-		<Do If="'$($(TestName)StagesWithProjectFile)' == 'True'">
-			<Command Name="BuildCookRun" Arguments="-Build -SkipCook -Stage -NoSubmit -Archive -ArchiveDirectory=$(LocalAchiveDirectory) -project=$($(TestName)Target) -configuration=$(Configuration) -platform=$(Platform) $($(Platform)$(TestName)ExtraCompilationArgs)"/>
-		</Do>
-		<Do If="'$($(TestName)StagesWithProjectFile)' != 'True'">
-			<Compile Target="$($(TestName)Target)" Configuration="$(Configuration)" Platform="$(Platform)" AllowXGE="false" Arguments="$($(Platform)$(TestName)ExtraCompilationArgs)"/>
-		</Do>
-	</Macro>
-
 	<Macro Name="DeployAndTest" Arguments="TestName;Platform">
-		<Property Name="PlatformStageDirName" Value="$(Platform)"/>
-		<Property Name="PlatformStageDirName" Value="Windows" If="'$(Platform)' == 'Win64'"/>
-
 		<Property Name="DevicePool" Value="$(GauntletDevicePool)" />
 		<Property Name="DevicePool" Value="$(AlternateGauntletDevicePool)" If="ContainsItem('$(AlternateDevicePoolPlatforms)', '$(Platform)', ';')" />
 		<Property Name="TestArgsExt" Value="$(TestArgs) $($(TestName)InitialExtraArgs) -deviceurl=$(GauntletDeviceService) -devicepool=$(DevicePool)"/>
@@ -134,9 +121,6 @@
 		<Property Name="TestArgsExt" Value="$(TestArgsExt) -tags=&quot;$(TestTags)&quot;"/>
 		<Property Name="TestArgsExt" Value="$(TestArgsExt) --buildmachine" If="$(IsBuildMachine)"/>
 		<Property Name="TestArgsExt" Value="$(TestArgsExt) -timeout=$(PerTestTimeout)"/>
-		<Property Name="TestArgsExt" Value="$(TestArgsExt) -SkipStage -Build=&quot;$(RootDir)/OnlineTests/$(PlatformStageDirName)&quot; $($(Platform)ExtraArgs)" If="'$($(TestName)StagesWithProjectFile)' == 'True'"/>
-		<Property Name="TestArgsExt" Value="$(TestArgsExt) -build=&quot;$($(TestName)BinariesRelative)\$(Platform)\$($(TestName)Target)$(TestModeSuffix)&quot;" If="'$($(TestName)StagesWithProjectFile)' != 'True'"/>
-		
 		<Property Name="TestArgsExt" Value="$(TestArgsExt) $($(Platform)ExtraArgs)"/>
 		<Do If="'$(Device)' != ''">
 			<Property Name="TestArgsExt" Value="$(TestArgsExt) -device=$(Device) "/>
@@ -151,7 +135,7 @@
 			<Property Name="TestArgsExt" Value=""/>
 		</Do>
 		<Command Name="RunLowLevelTests"
-				Arguments="-test=LowLevelTests -testapp=$($(TestName)Target)$(TestModeSuffix) -platform=$(Platform) -configuration=$(Configuration) $(TestArgsExt) -VeryVerbose"/>
+				Arguments="-test=LowLevelTests -testapp=$($(TestName)Target)$(TestModeSuffix) -platform=$(Platform) -configuration=$(Configuration) $(TestArgsExt) -build=&quot;$($(TestName)BinariesRelative)\$(Platform)\$($(TestName)Target)$(TestModeSuffix)&quot; -VeryVerbose"/>
 	</Macro>
 
 	<ForEach Name="BuildPlatform" Values="$(BuildPlatforms)">
@@ -214,12 +198,8 @@
 							<Property Name="TestRequirements" Value="$(TestRequirements);Catch2 Build Library $(Platform)"/>
 						</Do>
 						
-						<Node Name="$($(TestName)Short) Build Tests $(Platform)" Requires="$(TestRequirements)">
-							<Expand Name="BuildTest" TestName="$(TestName)" Platform="$(Platform)"/>
-						</Node>
-						<Label Category="Build" Name="$($(TestName)Short) Build $(Platform)" Requires="$($(TestName)Short) Build Tests $(Platform)"/>
-
-						<Node Name="$($(TestName)Short) Tests $(Platform)" Requires="$(TestRequirements);$($(TestName)Short) Build Tests $(Platform)">
+						<Node Name="$($(TestName)Short) Tests $(Platform)" Requires="$(TestRequirements)">
+							<Compile Target="$($(TestName)Target)" Configuration="$(Configuration)" Platform="$(Platform)" AllowXGE="false" Arguments="$($(Platform)$(TestName)ExtraCompilationArgs)"/>
 							<Do If="'$($(Platform)$(TestName)RunSupported)' == 'True'">
 								<Expand Name="DeployAndTest" TestName="$(TestName)" Platform="$(Platform)"/>
 							</Do>
