Add UnrealTraceServer to the Linux and Mac installed builds.

diff --git a/Engine/Build/InstalledEngineFilters.xml b/Engine/Build/InstalledEngineFilters.xml
--- a/Engine/Build/InstalledEngineFilters.xml
+++ b/Engine/Build/InstalledEngineFilters.xml
@@ -358,6 +358,9 @@
 	<Property Name="CopyEditorFilterMac" Value="$(CopyEditorFilter);$(MacBuildFiles)">
 		<!-- Extra tools -->
 		Engine/Extras/LLDBDataFormatters/UEDataFormatters_2ByteChars.py
+		
+		<!-- UnrealTraceServer is the tool that records Unreal Insights trace files -->
+		Engine/Binaries/Mac/UnrealTraceServer
 	</Property>
 
 	<Property Name="CopyEditorExceptionsMac" Value="$(CopyEditorExceptions)">
@@ -412,6 +415,9 @@
 
 		<!-- Bundled libcxx -->
 		Engine/Source/ThirdParty/Unix/LibCxx/...
+		
+		<!-- UnrealTraceServer is the tool that records Unreal Insights trace files -->
+		Engine/Binaries/Linux/UnrealTraceServer
 	</Property>
 	<Property Name="CopyEditorExceptionsLinux" Value="$(CopyEditorExceptions)">
 		<!-- Exclude Windows/Mac binaries on Linux here -->
@@ -606,6 +612,7 @@
 		Engine/Binaries/Linux/dump_syms.exe
 		Engine/Binaries/Linux/BreakpadSymbolEncoder.exe
 		Engine/Binaries/ThirdParty/Mono/Linux/...
+		Engine/Binaries/Linux/UnrealTraceServer
 	</Property>
 
 	<Property Name="CopyLinuxFilter">
