Add path for log linux
#rnx

diff --git a/Engine/Build/EditorPerf.xml b/Engine/Build/EditorPerf.xml
--- a/Engine/Build/EditorPerf.xml
+++ b/Engine/Build/EditorPerf.xml
@@ -45,6 +45,7 @@
 	<!-- Automation uses its own network output directory -->
 	<Property Name="AutomationNetworkRootDir" Value="\\epicgames.net\root\Builds\Automation\Reports\EditorPerf"/>
 	<Property Name="AutomationNetworkRootDir" Value="/Volumes/Builds/Automation/Reports/EditorPerf" If="'$(HostPlatform)' != 'Win64'"/>
+	<Property Name="AutomationNetworkRootDir" Value="/mnt/Builds/Automation/Reports/EditorPerf" If="'$(HostPlatform)' == 'Linux'"/>
 	<Property Name="AutomationLocalOutputDir" Value="$(RootDir)/LocalBuilds/Reports"/>
 
 	<!-- Outputs for this run keeping slashes flipped the right way because it shows up on the webapp this way -->
