Rebase Gauntlet EditorTests to UE.AutomationTestNode

#rnx

diff --git a/Engine/Build/EditorPerf.xml b/Engine/Build/EditorPerf.xml
--- a/Engine/Build/EditorPerf.xml
+++ b/Engine/Build/EditorPerf.xml
@@ -138,7 +138,7 @@
 	<Property Name="LazyLoadZenTestArgs" Value="$(LazyLoadTestArgs) -zenloader" />
 
 	<!-- Gauntlet arguments for all runs -->
-	<Property Name="GauntletArgs" Value="-build=editor -UseEditor -verbose -tempdir=&quot;$(RootDir)/Tests&quot;"/>
+	<Property Name="GauntletArgs" Value="-build=editor -verbose -tempdir=&quot;$(RootDir)/Tests&quot;"/>
 	<Property Name="GauntletArgs" Value="$(GauntletArgs) -branch=$(Branch) -ECBranch=&quot;$(Branch)&quot; -ECChangelist=&quot;$(Change)&quot;"/>
 	<Property Name="GauntletArgs" Value="$(GauntletArgs) -maxduration=28800"/>
 
@@ -149,6 +149,7 @@
 	<Property Name="ProfileArgs" Value="$(ProfileArgs) -GameInstanceLostTimerSeconds=3600"/>
 	<Property Name="ProfileArgs" Value="$(ProfileArgs) -NoLoadLevelAtStartup -Benchmarking -NODDCCLEANUP"/>
 	<Property Name="ProfileArgs" Value="$(ProfileArgs) -NoShaderDistrib -VerboseShaderLogging"/>
+	<Property Name="ProfileArgs" Value="$(ProfileArgs) -ResumeOnCriticalFailure"/>
 
 	<!-- options for summarize runs -->
 	<Option Name="RunTraceSummary" DefaultValue="false" Description="Summarize the insights trace after the run"/>
@@ -167,7 +168,7 @@
 	<Macro Name="RunTest" Arguments="TestNameSuffix;Platform;ProjectName;ProjectTarget;NodeName;Phase;Iteration;Args;EditorArgs">
 		<Property Name="TestName" Value="EditorPerf$(Phase)$(TestNameSuffix)"/>
 		<Property Name="TestFilter" Value="Project.Iteration.PIE"/>
-		<Property Name="TestFilterArgs" Value="-testfilter=$(TestFilter)"/>
+		<Property Name="TestFilterArgs" Value="-runtest=$(TestFilter)"/>
 		<Property Name="IterationSuffix" Value=""/>
 		<Property Name="IterationSuffix" Value="$(Iteration)" If="'$(Iteration)' != '1'"/>
 
@@ -205,7 +206,7 @@
 	<Macro Name="RunEditorWorkflowsTest" Arguments="Platform;ProjectName;NodeName">
 		<Property Name="TestName" Value="Editor.Workflows"/>
 		<Property Name="TestFilter" Value="Editor.Workflows"/>
-		<Property Name="TestFilterArgs" Value="-testfilter=$(TestFilter)"/>
+		<Property Name="TestFilterArgs" Value="-runtest=$(TestFilter)"/>
 
 		<Property Name="TestOutputDir" Value="$(AutomationOutputDir)/$(ProjectName)$(Platform)EditorWorkflows"/>
 
diff --git a/Engine/Source/Programs/AutomationTool/Gauntlet/Unreal/Editor/EditorTests.cs b/Engine/Source/Programs/AutomationTool/Gauntlet/Unreal/Editor/EditorTests.cs
--- a/Engine/Source/Programs/AutomationTool/Gauntlet/Unreal/Editor/EditorTests.cs
+++ b/Engine/Source/Programs/AutomationTool/Gauntlet/Unreal/Editor/EditorTests.cs
@@ -13,13 +13,8 @@ namespace UnrealEditor
 	/// Default set of options for testing Editor. Options that tests can configure
 	/// should be public, external command-line driven options should be protected/private
 	/// </summary>
-	public class EditorTestConfig : EngineTestConfig
+	public class EditorTestConfig : UE.AutomationTestConfig
 	{
-		/// <summary>
-		/// Filter or groups of tests to apply
-		/// </summary>
-		public override bool UseEditor { get; set; } = true;
-
 		/// <summary>
 		/// Use Simple Horde Report instead of Unreal Automated Tests
 		/// </summary>
@@ -138,7 +133,7 @@ namespace UnrealEditor
 		}
 	}
 
-	public class EditorTests : EngineTestBase<EditorTestConfig>
+	public class EditorTests : UE.AutomationNodeBase<EditorTestConfig>
 	{
 		public EditorTests(UnrealTestContext InContext) : base(InContext)
 		{
@@ -146,6 +141,10 @@ namespace UnrealEditor
 
 		public override EditorTestConfig GetConfiguration()
 		{
+			if (CachedConfig != null)
+			{
+				return CachedConfig;
+			}
 			// just need a single client
 			EditorTestConfig Config = base.GetConfiguration();
 			Config.RequireRole(UnrealTargetRole.Editor);	
