Fixing InstalledEngineBuild for Mac producing duplicate output errors when WithClient is enabled

#rnx
#prelight 64c80430a8539ecb49a30011

diff --git a/Engine/Build/InstalledEngineBuild.xml b/Engine/Build/InstalledEngineBuild.xml
--- a/Engine/Build/InstalledEngineBuild.xml
+++ b/Engine/Build/InstalledEngineBuild.xml
@@ -385,11 +385,9 @@
 			<ForEach Name="Tag" Values="$(OptionalTargetTags)" If="'$(OptionalTargetTags)' != ''">
 				<Tag Files="$(Tag)" With="#UnrealGame Mac"/>
 			</ForEach>
-			<ForEach Name="Target" Values="UnrealGame;$(OptionalClientTarget);$(OptionalServerTarget)">
-				<ForEach Name="Configuration" Values="$(GameConfigurations)">
-					<Compile Target="$(Target)" Platform="Mac" Configuration="$(Configuration)" Tag="#UnrealGame Mac" AllowParallelExecutor="$(bAllowParallelExecutor)" Arguments="-precompile -allmodules -nolink $(TargetDebugInfoArgMac)"/>
-					<Compile Target="$(Target)" Platform="Mac" Configuration="$(Configuration)" Tag="#UnrealGame Mac" AllowParallelExecutor="$(bAllowParallelExecutor)" Arguments="-precompile $(TargetDebugInfoArgMac)" Clean="false"/>
-				</ForEach>
+			<ForEach Name="Configuration" Values="$(GameConfigurations)">
+				<Compile Target="UnrealGame" Platform="Mac" Configuration="$(Configuration)" Tag="#UnrealGame Mac" AllowParallelExecutor="$(bAllowParallelExecutor)" Arguments="-precompile -allmodules -nolink $(TargetDebugInfoArgMac)"/>
+				<Compile Target="UnrealGame" Platform="Mac" Configuration="$(Configuration)" Tag="#UnrealGame Mac" AllowParallelExecutor="$(bAllowParallelExecutor)" Arguments="-precompile $(TargetDebugInfoArgMac)" Clean="false"/>
 			</ForEach>
 
 			<Tag Files="#UnrealGame Mac" Filter="$(MacStripFilter)" With="#UnrealGame Mac Unstripped"/>
