PR #10328: Copy Engine/Build/Turnkey directory to installed builds

#rnx

diff --git a/Engine/Build/InstalledEngineFilters.xml b/Engine/Build/InstalledEngineFilters.xml
--- a/Engine/Build/InstalledEngineFilters.xml
+++ b/Engine/Build/InstalledEngineFilters.xml
@@ -140,6 +140,9 @@
 		<!-- Build files -->
 		Engine/Build/Build.version
 		Engine/Build/Target.cs.template
+
+		<!-- Turnkey files -->
+		Engine/Build/Turnkey/...
 		
 		<!-- Source code and demos for container images -->
 		Engine/Extras/Containers/...
