Copy "index.json"

diff --git a/Engine/Build/EditorPerf.xml b/Engine/Build/EditorPerf.xml
--- a/Engine/Build/EditorPerf.xml
+++ b/Engine/Build/EditorPerf.xml
@@ -229,9 +229,13 @@
 	</Macro>
 
 	<Macro Name="UploadEditorWorkflowsPerformanceDataAndTraces" Arguments="Platform;ProjectName">
-		<!-- Copy the utrace file and the generated csv performance data to the output directory -->
 		<Property Name="TestOutputDir" Value="$(AutomationOutputDir)/$(ProjectName)$(Platform)EditorWorkflows"/>
-		<Copy From="$(TraceDir)/Editor.Workflows/*" To="$(TestOutputDir)/Editor.Workflows"/>
+
+		<!-- Copy the utrace file and the generated csv performance data to the output directory -->
+		<Copy From="$(TraceDir)/Editor.Workflows/*" To="$(TestOutputDir)/Editor.Workflows" Overwrite="true"/>
+
+		<!-- Copy "index.json" that captures the functional results to the output directory -->
+		<Copy From="$(RootDir)/Tests/TestReport/index.json" To="$(TestOutputDir)/Editor.Workflows/index.json" Overwrite="true"/>
 	</Macro>
 
 	<!-- These will be set to the names of all nodes we want to execute based on specified platforms/options -->
