Make SummarizeTraceEditorWorkflows part of UploadEditorWorkflowsPerformanceDataAndTraces an after step not conditional to previous step success state

#rnx

diff --git a/Engine/Build/EditorPerf.xml b/Engine/Build/EditorPerf.xml
--- a/Engine/Build/EditorPerf.xml
+++ b/Engine/Build/EditorPerf.xml
@@ -205,21 +205,16 @@
 	<!-- The main editor workflows test running macro logic that is expanded with platform and dependency node names -->
 	<Macro Name="RunEditorWorkflowsTest" Arguments="Platform;ProjectName;NodeName">
 		<Property Name="TestName" Value="Editor.Workflows"/>
-		<Property Name="TestFilter" Value="Editor.Workflows"/>
-		<Property Name="TestFilterArgs" Value="-runtest=$(TestFilter)"/>
+		<Property Name="TestFilterArgs" Value="-runtest=$(TestName)"/>
 
 		<Property Name="TestOutputDir" Value="$(AutomationOutputDir)/$(ProjectName)$(Platform)EditorWorkflows"/>
+		<Property Name="TraceFilePath" Value="$(TraceDir)/$(TestName)/$(ProjectName)$(Platform)EditorWorkflows.utrace" />
 
 		<Property Name="LogOutputParams" Value="-logdir=&quot;$(TestOutputDir)&quot;"/>
 		<Property Name="ReportParams" Value="-WriteTestResultsForHorde -HordeTestDataKey=&quot;$(NodeName)$(PreflightSuffix)&quot;"/>
 		<Property Name="AllGauntletArgs" Value="-Project=$(ProjectName) -Platform=$(Platform) $(GauntletArgs) $(ReportParams) $(LogOutputParams)"/>
 
-		<Property Name="TraceFilePath" Value="$(TraceDir)/Editor.Workflows/$(ProjectName)$(Platform)EditorWorkflows.utrace" />
-
 		<Command Name="RunUnreal" Arguments="$(AllGauntletArgs) $(ProfileArgs) $(TestFilterArgs) -sessionlabel=&quot;$(TestName)&quot; -tracefile=&quot;$(TraceFilePath)&quot;" />
-
-		<!-- Translate the utrace file into csv data -->
-		<Commandlet Project="$(ProjectName)" Name="SummarizeTraceEditorWorkflows" Arguments="-inputfile=&quot;$(TraceFilePath)&quot; -testname=&quot;$(TestName)&quot; $(CommandletArgs)" />
 	</Macro>
 	
 	<Macro Name="UploadTestTraces" Arguments="Platform;ProjectName">		
@@ -230,6 +225,10 @@
 
 	<Macro Name="UploadEditorWorkflowsPerformanceDataAndTraces" Arguments="Platform;ProjectName">
 		<Property Name="TestOutputDir" Value="$(AutomationOutputDir)/$(ProjectName)$(Platform)EditorWorkflows"/>
+		<Property Name="TraceFilePath" Value="$(TraceDir)/Editor.Workflows/$(ProjectName)$(Platform)EditorWorkflows.utrace" />
+
+		<!-- Translate the utrace file into csv data -->
+		<Commandlet Project="$(ProjectName)" Name="SummarizeTraceEditorWorkflows" Arguments="-inputfile=&quot;$(TraceFilePath)&quot; -testname=&quot;Editor.Workflows&quot; $(CommandletArgs)" />
 
 		<!-- Copy the utrace file and the generated csv performance data to the output directory -->
 		<Copy From="$(TraceDir)/Editor.Workflows/*" To="$(TestOutputDir)/Editor.Workflows" Overwrite="true"/>
@@ -380,22 +379,23 @@
 	<!-- Generate graph nodes for each platform to run -->
 	<ForEach Name="Platform" Values="$(Platforms)">
 
+		<Property Name="ProjectName" Value="EngineTest" />
 		<!-- Define output name for the EditorWorkflows test -->
-		<Property Name="EditorWorkflowsOutputName" Value="EngineTest$(Platform)EditorWorkflows"/>
-		<Property Name="CompileNodeName" Value="Compile EngineTest $(Platform) EditorWorkflows"/>
+		<Property Name="EditorWorkflowsOutputName" Value="$(ProjectName)$(Platform)EditorWorkflows"/>
+		<Property Name="CompileNodeName" Value="Compile $(ProjectName) $(Platform) EditorWorkflows"/>
 		<Property Name="AgentModifier" Value=""/>
-		<Property Name="AgentModifier" Value=" EngineTest" If="!$(SingleAgentCompile)"/>
+		<Property Name="AgentModifier" Value=" $(ProjectName)" If="!$(SingleAgentCompile)"/>
 
 		<!-- Compile Agent configuration -->
 		<Agent Name="Compile Editor Agent $(Platform)$(AgentModifier)" Type="$($(Platform)CompileAgentType)">
 			<Node Name="$(CompileNodeName)">
 				<SetVersion Change="$(Change)" Branch="$(EscapedBranch)" If="$(IsBuildMachine)"/>
 				<Compile Target="ShaderCompileWorker" Platform="$(Platform)" Configuration="Development"/>
-				<Compile Target="CrashReportClientEditor" Platform="$(Platform)" Project="EngineTest" Configuration="Shipping"/>
-				<Compile Target="EngineTestEditor" Platform="$(Platform)" Project="EngineTest" Configuration="Development" Arguments="$(EditorCompileArgs)"/>
+				<Compile Target="CrashReportClientEditor" Platform="$(Platform)" Project="$(ProjectName)" Configuration="Shipping"/>
+				<Compile Target="$(ProjectName)Editor" Platform="$(Platform)" Project="$(ProjectName)" Configuration="Development" Arguments="$(EditorCompileArgs)"/>
 			</Node>
 		</Agent>
-		<Label Category="$(CompileCategory)" Name="EngineTest Editor Compile" Requires="$(CompileNodeName)"/>
+		<Label Category="$(CompileCategory)" Name="$(ProjectName) Editor Compile" Requires="$(CompileNodeName)"/>
 
 		<!-- Set up properties for run and upload nodes -->
 		<Property Name="RunEditorWorkflowsNodeRequires" Value="$(CompileNodeName)"/>
@@ -403,10 +403,10 @@
 		<Property Name="RunEditorWorkflowsNodeList" Value=""/>
 
 		<!-- Configure EditorWorkflows Agent -->
-		<Agent Name="EditorWorkflows Agent $(Platform) EngineTest" Type="EditorPerf$(Platform)">
+		<Agent Name="EditorWorkflows Agent $(Platform) $(ProjectName)" Type="EditorPerf$(Platform)">
 			<Property Name="RunEditorWorkflowsNodeName" Value="Run $(EditorWorkflowsOutputName)" />
 			<Node Name="$(RunEditorWorkflowsNodeName)" Requires="$(RunEditorWorkflowsNodeRequires);$(RunEditorWorkflowsNodeList)">
-				<Expand Name="RunEditorWorkflowsTest" Platform="$(Platform)" ProjectName="EngineTest" NodeName="$(RunEditorWorkflowsNodeName)"/>
+				<Expand Name="RunEditorWorkflowsTest" Platform="$(Platform)" ProjectName="$(ProjectName)" NodeName="$(RunEditorWorkflowsNodeName)"/>
 			</Node>
 
 			<Property Name="RunEditorWorkflowsNodeList" Value="$(RunEditorWorkflowsNodeList);$(RunEditorWorkflowsNodeName)" />
@@ -414,13 +414,13 @@
 
 			<!-- Set up Node for uploading traces -->
 			<Node Name="$(RunNodeUploadTraces)">
-				<Expand Name="UploadEditorWorkflowsPerformanceDataAndTraces" Platform="$(Platform)" ProjectName="EngineTest"/>
+				<Expand Name="UploadEditorWorkflowsPerformanceDataAndTraces" Platform="$(Platform)" ProjectName="$(ProjectName)"/>
 			</Node>
 
 			<!-- Update UploadEditorWorkflowsNodeList -->
 			<Property Name="UploadEditorWorkflowsNodeList" Value="$(UploadEditorWorkflowsNodeList);$(RunNodeUploadTraces)"/>
 		</Agent>
-		<Label Category="$(RunCategory)" Name="EngineTest Editor Run" Requires="$(RunEditorWorkflowsNodeRequires);$(RunEditorWorkflowsNodeList);$(UploadEditorWorkflowsNodeList)" UgsBadge="EngineTest Editor Run"/>
+		<Label Category="$(RunCategory)" Name="$(ProjectName) Editor Run" Requires="$(RunEditorWorkflowsNodeRequires);$(RunEditorWorkflowsNodeList);$(UploadEditorWorkflowsNodeList)" UgsBadge="$(ProjectName) Editor Run"/>
 		<Property Name="RunEditorWorkflowsGraphNodes" Value="$(RunEditorWorkflowsGraphNodes);$(RunEditorWorkflowsNodeList);$(UploadEditorWorkflowsNodeList)"/>
 	</ForEach>
 	<!-- End ForEach Platforms -->
