Apply fixes to OnlineTests to allow correct staging with BuildCookRun on Windows
- LowLevelTests.xml Builds tests under BuildCookRun with -NoBootstrapExe
- TestModuleRules no longer generates arguments -BaseFromWorkingDir to generated .xml files
- Fixes bug in RunLowLevelTests that puts tags in the wrong position in command line arguments

diff --git a/Engine/Build/LowLevelTests.xml b/Engine/Build/LowLevelTests.xml
--- a/Engine/Build/LowLevelTests.xml
+++ b/Engine/Build/LowLevelTests.xml
@@ -109,11 +109,22 @@
 		<Label Category="After" Name="$($(TestName)Short) After $(Platform)" Requires="$(TestName) AfterStep $(Platform)"/>
 		<Property Name="AfterTestNodes" Value="$(AfterTestNodes);$(TestName) AfterStep $(Platform)" If="$(RunAllTests) Or $(Run$(TestName)Tests)"/>
 	</Macro>
-
+	
+	<Macro Name="BuildAndStageTest" Arguments="Platform;TestName">
+		<Property Name="LocalAchiveDirectory" Value="$(RootDir)/$($(TestName)Target)"/>
+		<Do If="'$($(TestName)StagesWithProjectFile)' == 'True'">
+			<Command Name="BuildCookRun" Arguments="-Build -NoBootstrapExe -NoCodeSign -SkipCook -Stage -NoSubmit -Archive -ArchiveDirectory=$(LocalAchiveDirectory) -project=$($(TestName)Target) -configuration=$(Configuration) -platform=$(Platform) $($(Platform)$(TestName)ExtraCompilationArgs)"/>
+		</Do>
+		<Do If="'$($(TestName)StagesWithProjectFile)' != 'True'">
+			<Compile Target="$($(TestName)Target)" Configuration="$(Configuration)" Platform="$(Platform)" AllowXGE="false" Arguments="$($(Platform)$(TestName)ExtraCompilationArgs)"/>
+		</Do>
+	</Macro>
 	<Macro Name="DeployAndTest" Arguments="TestName;Platform">
+		<Property Name="PlatformStageDirName" Value="$(Platform)"/>
+		<Property Name="PlatformStageDirName" Value="Windows" If="'$(Platform)' == 'Win64'"/>
 		<Property Name="DevicePool" Value="$(GauntletDevicePool)" />
 		<Property Name="DevicePool" Value="$(AlternateGauntletDevicePool)" If="ContainsItem('$(AlternateDevicePoolPlatforms)', '$(Platform)', ';')" />
-		<Property Name="TestArgsExt" Value="$(TestArgs) $($(TestName)InitialExtraArgs) -deviceurl=$(GauntletDeviceService) -devicepool=$(DevicePool)"/>
+		<Property Name="TestArgsExt" Value="$(TestArgs) -deviceurl=$(GauntletDeviceService) -devicepool=$(DevicePool)"/>
 		<Property Name="TestModeSuffix" Value=""/>
 		<Property Name="TestTags" Value=""/>
 		<Property Name="TestTags" Value="$($(Platform)$(TestName)Tags)" If="'$($(Platform)$(TestName)Tags)' != ''"/>
@@ -121,6 +132,8 @@
 		<Property Name="TestArgsExt" Value="$(TestArgsExt) -tags=&quot;$(TestTags)&quot;"/>
 		<Property Name="TestArgsExt" Value="$(TestArgsExt) --buildmachine" If="$(IsBuildMachine)"/>
 		<Property Name="TestArgsExt" Value="$(TestArgsExt) -timeout=$(PerTestTimeout)"/>
+		<Property Name="TestArgsExt" Value="$(TestArgsExt) $($(TestName)GauntletArgs)"/>
+		<Property Name="TestArgsExt" Value="$(TestArgsExt) -extra-args=&quot;$($(TestName)ExtraArgs)&quot;"/>
 		<Property Name="TestArgsExt" Value="$(TestArgsExt) $($(Platform)ExtraArgs)"/>
 		<Do If="'$(Device)' != ''">
 			<Property Name="TestArgsExt" Value="$(TestArgsExt) -device=$(Device) "/>
@@ -134,8 +147,12 @@
 		<Do If="'$($(TestName)UsesCatch2)' != 'True'">
 			<Property Name="TestArgsExt" Value=""/>
 		</Do>
+
+		<Property Name="TestArgsExt" Value="$(TestArgsExt) -SkipStage -Build=&quot;$(RootDir)\$($(TestName)Target)\$(PlatformStageDirName)&quot;" If="'$($(TestName)StagesWithProjectFile)' == 'True'"/>
+		<Property Name="TestArgsExt" Value="$(TestArgsExt) -build=&quot;$($(TestName)BinariesRelative)\$(Platform)\$($(TestName)Target)$(TestModeSuffix)&quot;" If="'$($(TestName)StagesWithProjectFile)' != 'True'"/>
+		
 		<Command Name="RunLowLevelTests"
-				Arguments="-test=LowLevelTests -testapp=$($(TestName)Target)$(TestModeSuffix) -platform=$(Platform) -configuration=$(Configuration) $(TestArgsExt) -build=&quot;$($(TestName)BinariesRelative)\$(Platform)\$($(TestName)Target)$(TestModeSuffix)&quot; -VeryVerbose"/>
+				Arguments="-test=LowLevelTests -testapp=$($(TestName)Target)$(TestModeSuffix) -platform=$(Platform) -configuration=$(Configuration) $(TestArgsExt) -VeryVerbose"/>
 	</Macro>
 
 	<ForEach Name="BuildPlatform" Values="$(BuildPlatforms)">
@@ -197,10 +214,9 @@
 						<Do If="$(BuildCatch2Libs)">
 							<Property Name="TestRequirements" Value="$(TestRequirements);Catch2 Build Library $(Platform)"/>
 						</Do>
-						
 						<Node Name="$($(TestName)Short) Tests $(Platform)" Requires="$(TestRequirements)">
-							<Compile Target="$($(TestName)Target)" Configuration="$(Configuration)" Platform="$(Platform)" AllowXGE="false" Arguments="$($(Platform)$(TestName)ExtraCompilationArgs)"/>
 							<Do If="'$($(Platform)$(TestName)RunSupported)' == 'True'">
+								<Expand Name="BuildAndStageTest" TestName="$(TestName)" Platform="$(Platform)"/>
 								<Expand Name="DeployAndTest" TestName="$(TestName)" Platform="$(Platform)"/>
 							</Do>
 						</Node>
@@ -218,4 +234,4 @@
 	
 	<Aggregate Name="Run Low Level Tests" Requires="$(TestNodes);$(AfterTestNodes)"/>
 	<Aggregate Name="Build Catch2 Libraries" Requires="$(Catch2BuildNodes)"/>
-</BuildGraph>
+</BuildGraph>
\ No newline at end of file
diff --git a/Engine/Build/LowLevelTests/Foundation.xml b/Engine/Build/LowLevelTests/Foundation.xml
--- a/Engine/Build/LowLevelTests/Foundation.xml
+++ b/Engine/Build/LowLevelTests/Foundation.xml
@@ -8,6 +8,8 @@
   <Property Name="FoundationTarget" Value="FoundationTests" />
   <Property Name="FoundationBinariesRelative" Value="Engine\Binaries" />
   <Property Name="FoundationReportType" Value="xml" />
+  <Property Name="FoundationGauntletArgs" Value="" />
+  <Property Name="FoundationExtraArgs" />
   <Property Name="FoundationInitialExtraArgs" />
   <Property Name="FoundationHasAfterSteps" Value="False" />
   <Property Name="FoundationUsesCatch2" Value="True" />
@@ -34,4 +36,7 @@
   <Property Name="LinuxFoundationRunSupported" Value="True" />
   <Property Name="LinuxArm64FoundationRunSupported" Value="True" />
   <Property Name="TVOSFoundationRunSupported" Value="True" />
+  <Property Name="VisionOSFoundationTags" Value="" />
+  <Property Name="VisionOSFoundationExtraCompilationArgs" Value="-allmodules" />
+  <Property Name="VisionOSFoundationRunSupported" Value="True" />
 </BuildGraph>
\ No newline at end of file
diff --git a/Engine/Build/LowLevelTests/HeadlessChaos.xml b/Engine/Build/LowLevelTests/HeadlessChaos.xml
--- a/Engine/Build/LowLevelTests/HeadlessChaos.xml
+++ b/Engine/Build/LowLevelTests/HeadlessChaos.xml
@@ -8,6 +8,8 @@
   <Property Name="HeadlessChaosTarget" Value="HeadlessChaos" />
   <Property Name="HeadlessChaosBinariesRelative" Value="Engine\Binaries" />
   <Property Name="HeadlessChaosReportType" Value="console" />
+  <Property Name="HeadlessChaosGauntletArgs" Value="" />
+  <Property Name="HeadlessChaosExtraArgs" />
   <Property Name="HeadlessChaosInitialExtraArgs" />
   <Property Name="HeadlessChaosHasAfterSteps" Value="False" />
   <Property Name="HeadlessChaosUsesCatch2" Value="False" />
@@ -34,4 +36,7 @@
   <Property Name="LinuxHeadlessChaosRunSupported" Value="True" />
   <Property Name="LinuxArm64HeadlessChaosRunSupported" Value="True" />
   <Property Name="TVOSHeadlessChaosRunSupported" Value="True" />
+  <Property Name="VisionOSHeadlessChaosTags" Value="" />
+  <Property Name="VisionOSHeadlessChaosExtraCompilationArgs" Value="" />
+  <Property Name="VisionOSHeadlessChaosRunSupported" Value="True" />
 </BuildGraph>
\ No newline at end of file
diff --git a/Engine/Build/LowLevelTests/NetworkPredictionPlugin.xml b/Engine/Build/LowLevelTests/NetworkPredictionPlugin.xml
--- a/Engine/Build/LowLevelTests/NetworkPredictionPlugin.xml
+++ b/Engine/Build/LowLevelTests/NetworkPredictionPlugin.xml
@@ -8,6 +8,8 @@
   <Property Name="NetworkPredictionPluginTarget" Value="NetworkPredictionTests" />
   <Property Name="NetworkPredictionPluginBinariesRelative" Value="Engine\Binaries" />
   <Property Name="NetworkPredictionPluginReportType" Value="console" />
+  <Property Name="NetworkPredictionPluginGauntletArgs" Value="" />
+  <Property Name="NetworkPredictionPluginExtraArgs" />
   <Property Name="NetworkPredictionPluginInitialExtraArgs" />
   <Property Name="NetworkPredictionPluginHasAfterSteps" Value="False" />
   <Property Name="NetworkPredictionPluginUsesCatch2" Value="True" />
@@ -34,4 +36,7 @@
   <Property Name="LinuxNetworkPredictionPluginRunSupported" Value="True" />
   <Property Name="LinuxArm64NetworkPredictionPluginRunSupported" Value="True" />
   <Property Name="TVOSNetworkPredictionPluginRunSupported" Value="True" />
+  <Property Name="VisionOSNetworkPredictionPluginTags" Value="" />
+  <Property Name="VisionOSNetworkPredictionPluginExtraCompilationArgs" Value="" />
+  <Property Name="VisionOSNetworkPredictionPluginRunSupported" Value="True" />
 </BuildGraph>
\ No newline at end of file
diff --git a/Engine/Build/LowLevelTests/ReplicationSystem.xml b/Engine/Build/LowLevelTests/ReplicationSystem.xml
--- a/Engine/Build/LowLevelTests/ReplicationSystem.xml
+++ b/Engine/Build/LowLevelTests/ReplicationSystem.xml
@@ -8,6 +8,8 @@
   <Property Name="ReplicationSystemTarget" Value="ReplicationSystemLowLevelTests" />
   <Property Name="ReplicationSystemBinariesRelative" Value="Engine\Binaries" />
   <Property Name="ReplicationSystemReportType" Value="console" />
+  <Property Name="ReplicationSystemGauntletArgs" Value="" />
+  <Property Name="ReplicationSystemExtraArgs" />
   <Property Name="ReplicationSystemInitialExtraArgs" />
   <Property Name="ReplicationSystemHasAfterSteps" Value="False" />
   <Property Name="ReplicationSystemUsesCatch2" Value="True" />
@@ -34,4 +36,7 @@
   <Property Name="LinuxReplicationSystemRunSupported" Value="True" />
   <Property Name="LinuxArm64ReplicationSystemRunSupported" Value="True" />
   <Property Name="TVOSReplicationSystemRunSupported" Value="True" />
+  <Property Name="VisionOSReplicationSystemTags" Value="" />
+  <Property Name="VisionOSReplicationSystemExtraCompilationArgs" Value="" />
+  <Property Name="VisionOSReplicationSystemRunSupported" Value="True" />
 </BuildGraph>
\ No newline at end of file
diff --git a/Engine/Source/Programs/AutomationTool/LowLevelTests/RunLowLevelTests.cs b/Engine/Source/Programs/AutomationTool/LowLevelTests/RunLowLevelTests.cs
--- a/Engine/Source/Programs/AutomationTool/LowLevelTests/RunLowLevelTests.cs
+++ b/Engine/Source/Programs/AutomationTool/LowLevelTests/RunLowLevelTests.cs
@@ -171,6 +171,9 @@ namespace LowLevelTests
 
 		public string Build;
 
+		[AutoParam("")]
+		public string TestExtraArgs;
+
 		[AutoParam("")]
 		public string LogDir;
 
@@ -231,6 +234,8 @@ namespace LowLevelTests
 			Tags = Params.ParseValue("tags=", null);
 			AttachToDebugger = Params.ParseParam("attachtodebugger");
 
+			TestExtraArgs = Params.ParseValue("extra-args=", null);
+
 			SkipStage = Params.ParseParam("skipstage");
 
 			ReportType = Params.ParseValue("reporttype=", null);
@@ -246,6 +251,7 @@ namespace LowLevelTests
 					&& !Arg.StartsWith("platform=", StringComparison.OrdinalIgnoreCase)
 					&& !Arg.StartsWith("device=", StringComparison.OrdinalIgnoreCase))
 				.ToArray();
+
 			Params = new Params(CleanArgs);
 		}
 	}
@@ -262,10 +268,11 @@ namespace LowLevelTests
 		private bool AttachToDebugger { get; set; }
 		private string ReportType { get; set; }
 		private int PerTestTimeout { get; set; }
+		private string TestExtraArgs { get; set; }
 
 		public UnrealDeviceReservation UnrealDeviceReservation { get; private set; }
 
-		public LowLevelTestsSession(LowLevelTestsBuildSource InBuildSource, string InTags, int InSleep, bool InAttachToDebugger, string InReportType, int InPerTestTimeout = 0)
+		public LowLevelTestsSession(LowLevelTestsBuildSource InBuildSource, string InTags, int InSleep, bool InAttachToDebugger, string InReportType, int InPerTestTimeout = 0, string InTestExtraArgs = null)
 		{
 			BuildSource = InBuildSource;
 			Tags = InTags;
@@ -274,6 +281,7 @@ namespace LowLevelTests
 			ReportType = InReportType;
 			PerTestTimeout = InPerTestTimeout;
 			UnrealDeviceReservation = new UnrealDeviceReservation();
+			TestExtraArgs = InTestExtraArgs;
 		}
 
 		public bool TryReserveDevices()
@@ -296,7 +304,7 @@ namespace LowLevelTests
 
 			// TargetDevice<Platform> classes have a hard dependency on UnrealAppConfig instead of IAppConfig.
 			// More refactoring needed to support non-packaged applications that can be run natively from a path on the device.
-			UnrealAppConfig AppConfig = BuildSource.GetUnrealAppConfig(Tags, Sleep, AttachToDebugger, ReportType, PerTestTimeout);
+			UnrealAppConfig AppConfig = BuildSource.GetUnrealAppConfig(Tags, Sleep, AttachToDebugger, ReportType, PerTestTimeout, TestExtraArgs);
 
 			IEnumerable<ITargetDevice> DevicesToInstallOn = UnrealDeviceReservation.ReservedDevices.ToArray();
 			ITargetDevice Device = DevicesToInstallOn.Where(D => D.IsConnected && D.Platform == BuildSource.Platform).First();
@@ -676,7 +684,7 @@ namespace LowLevelTests
 				.ToArray();
 		}
 
-		public UnrealAppConfig GetUnrealAppConfig(string InTags, int InSleep, bool InAttachToDebugger, string InReportType, int InPerTestTimeout = 0)
+		public UnrealAppConfig GetUnrealAppConfig(string InTags, int InSleep, bool InAttachToDebugger, string InReportType, int InPerTestTimeout = 0, string TestExtraArgs = null)
 		{
 			if (CachedConfig == null)
 			{
@@ -689,6 +697,13 @@ namespace LowLevelTests
 				CachedConfig.Build = DiscoveredBuild;
 				CachedConfig.Sandbox = "LowLevelTests";
 				CachedConfig.FilesToCopy = new List<UnrealFileToCopy>();
+
+				//Tags needs to be the first argument, if any are provided via --tags=
+				if (!string.IsNullOrEmpty(InTags))
+				{
+					CachedConfig.CommandLineParams.Add(InTags, null, true);
+				}
+
 				// Set reporting options, filters etc
 				CachedConfig.CommandLineParams.AddRawCommandline("--durations=no");
 				if (!string.IsNullOrEmpty(InReportType))
@@ -697,10 +712,6 @@ namespace LowLevelTests
 					CachedConfig.CommandLineParams.AddRawCommandline(string.Format("--out={0}", LowLevelTestsReporting.GetTargetReportPath(Platform, TestApp, BuildPath)));
 				}
 				CachedConfig.CommandLineParams.AddRawCommandline("--filenames-as-tags");
-				if (!string.IsNullOrEmpty(InTags))
-				{
-					CachedConfig.CommandLineParams.Add(InTags, null, true);
-				}
 				if (InSleep > 0)
 				{
 					CachedConfig.CommandLineParams.AddRawCommandline(String.Format("--sleep={0}", InSleep));
@@ -732,10 +743,16 @@ namespace LowLevelTests
 					}
 				}
 
+				if (TestExtraArgs != null)
+				{
+					ExtraCmd += string.Format(" {0}", TestExtraArgs);
+				}
+
 				if (!string.IsNullOrEmpty(ExtraCmd))
 				{
 					CachedConfig.CommandLineParams.AddRawCommandline("--extra-args" + ExtraCmd);
-				}
+				}				
+
 
 				CachedConfig.CanAlterCommandArgs = false; // No further changes by IAppInstall instances etc.
 			}
diff --git a/Engine/Source/Programs/AutomationTool/LowLevelTests/Tests/Gauntlet.LowLevelTestNode.cs b/Engine/Source/Programs/AutomationTool/LowLevelTests/Tests/Gauntlet.LowLevelTestNode.cs
--- a/Engine/Source/Programs/AutomationTool/LowLevelTests/Tests/Gauntlet.LowLevelTestNode.cs
+++ b/Engine/Source/Programs/AutomationTool/LowLevelTests/Tests/Gauntlet.LowLevelTestNode.cs
@@ -50,7 +50,7 @@ namespace LowLevelTests
 		{
 			if (LowLevelTestsApp == null)
 			{
-				LowLevelTestsApp = new LowLevelTestsSession(Context.BuildInfo, Context.Options.Tags, Context.Options.Sleep, Context.Options.AttachToDebugger, Context.Options.ReportType, Context.Options.Timeout);
+				LowLevelTestsApp = new LowLevelTestsSession(Context.BuildInfo, Context.Options.Tags, Context.Options.Sleep, Context.Options.AttachToDebugger, Context.Options.ReportType, Context.Options.Timeout, Context.Options.TestExtraArgs);
 			}
 
 			return LowLevelTestsApp.TryReserveDevices();
diff --git a/Engine/Source/Programs/UnrealBuildTool/Configuration/TestModuleRules.cs b/Engine/Source/Programs/UnrealBuildTool/Configuration/TestModuleRules.cs
--- a/Engine/Source/Programs/UnrealBuildTool/Configuration/TestModuleRules.cs
+++ b/Engine/Source/Programs/UnrealBuildTool/Configuration/TestModuleRules.cs
@@ -238,15 +238,8 @@ namespace UnrealBuildTool
 					InsertOrUpdateTestFlagProperty(ref lastUpdatedNode, TestMetadata.TestName, "Target", Convert.ToString(TestTargetName));
 					InsertOrUpdateTestFlagProperty(ref lastUpdatedNode, TestMetadata.TestName, "BinariesRelative", Convert.ToString(TestBinariesPath));
 					InsertOrUpdateTestFlagProperty(ref lastUpdatedNode, TestMetadata.TestName, "ReportType", Convert.ToString(TestMetadata.ReportType));
-
-					string ExtraArguments = Convert.ToString(TestMetadata.InitialExtraArgs);
-					if(TestMetadata.StagesWithProjectFile)
-					{
-						//  -BaseFromWorkingDir Is needed here for Programs to properly detect .ini files from its irregular platform path.
-						ExtraArguments += $"--extra-args -BaseFromWorkingDir";
-					}
-
-					InsertOrUpdateTestFlagProperty(ref lastUpdatedNode, TestMetadata.TestName, "InitialExtraArgs", ExtraArguments);
+					InsertOrUpdateTestFlagProperty(ref lastUpdatedNode, TestMetadata.TestName, "GauntletArgs", Convert.ToString(TestMetadata.InitialExtraArgs) + Convert.ToString(TestMetadata.GauntletArgs));
+					InsertOrUpdateTestFlagProperty(ref lastUpdatedNode, TestMetadata.TestName, "ExtraArgs", Convert.ToString(TestMetadata.ExtraArgs));
 					InsertOrUpdateTestFlagProperty(ref lastUpdatedNode, TestMetadata.TestName, "HasAfterSteps", Convert.ToString(TestMetadata.HasAfterSteps));
 					InsertOrUpdateTestFlagProperty(ref lastUpdatedNode, TestMetadata.TestName, "UsesCatch2", Convert.ToString(TestMetadata.UsesCatch2));
 
@@ -380,7 +373,6 @@ namespace UnrealBuildTool
 				Document.Root!.Add(ElementAppend);
 			}
 		}
-
 		private void InsertOrUpdateTestFlagProperty(ref XElement ElementUpsertAfter, string TestName, string FlagSuffix, string FlagValue)
 		{
 			IEnumerable<XElement> NextChunk = ElementUpsertAfter.ElementsAfterSelf(BuildGraphNamespace + "Property")
@@ -453,8 +445,7 @@ namespace UnrealBuildTool
 			}
 
 			/// <summary>
-			/// Does this test use project files for staging additional files?
-			/// This will append `-SkipStage -Build=$(RootDir)/$($(TestName)Target)` to InitialExtraArgs
+			/// Does this test use project files for staging additional files
 			/// and cause the build to use BuildCookRun instead of a Compile step
 			/// </summary>
 			public bool StagesWithProjectFile { get; set; }
@@ -464,10 +455,25 @@ namespace UnrealBuildTool
 			/// </summary>
 			public bool Disabled { get; set; }
 
+			/// <summary>
+			/// Depercated, use GauntletArgs or ExtraArgs instead to help indicate arguments to launch the test under.
+			/// </summary>
+			public string InitialExtraArgs
+			{ 
+				get;
+				[Obsolete]
+				set; 
+			}
+
 			/// <summary>
 			/// Any initial extra args to be passed to the test executable
 			/// </summary>
-			public string InitialExtraArgs { get; set; }
+			public string GauntletArgs { get; set; }
+
+			/// <summary>
+			/// Any extra args to be passed to the test executable as --extra-args
+			/// </summary>
+			public string ExtraArgs { get; set; }
 
 			/// <summary>
 			/// Whether there's a step that gets executed after the tests have finished.
