FAB plugin - mask out the plugin from UE5 binary install

diff --git a/Engine/Build/InstalledEngineFilters.xml b/Engine/Build/InstalledEngineFilters.xml
--- a/Engine/Build/InstalledEngineFilters.xml
+++ b/Engine/Build/InstalledEngineFilters.xml
@@ -210,6 +210,7 @@
 		Engine/Programs/UnrealGameSync/...
 
 		<!-- Plugins -->
+		Engine/Plugins/Fab/...
 		Engine/Plugins/Bridge/...
 		Engine/Plugins/Runtime/TwitchLiveStreaming/...
 		Engine/Plugins/Runtime/PacketHandlers/CompressionComponents/Oodle/...
