Adding Vehicle CPP template to game templates list now that it has been updated in 5.3.

diff --git a/Engine/Build/InstalledEngineFilters.xml b/Engine/Build/InstalledEngineFilters.xml
--- a/Engine/Build/InstalledEngineFilters.xml
+++ b/Engine/Build/InstalledEngineFilters.xml
@@ -112,6 +112,7 @@
 		Templates/TP_ThirdPersonBP/...
 		Templates/TP_TopDown/...
 		Templates/TP_TopDownBP/...
+		Templates/TP_VehicleAdv/...
 		Templates/TP_VehicleAdvBP/...
 		Templates/TP_VirtualRealityBP/...
 
