Add option disable/enable running after steps in LowLevelTests.xml.

diff --git a/Engine/Build/LowLevelTests.xml b/Engine/Build/LowLevelTests.xml
--- a/Engine/Build/LowLevelTests.xml
+++ b/Engine/Build/LowLevelTests.xml
@@ -102,6 +102,7 @@
 	<Option Name="PerTestTimeout" DefaultValue="0" Description="Individual test timeout (not for entire test run). 0 means no timeout is captured."/>
 	
 	<!-- After Steps -->
+	<Option Name="RunAfterSteps" DefaultValue="false" Description="If set RunAfterSteps is True then we should run aftersteps for this run."/>
 	<Macro Name="ExpandAfterSteps" Arguments="Platform;TestName">
 		<Node Name="$(TestName) AfterStep $(Platform)">
 			<Expand Name="$(TestName)AfterSteps" Platform="$(Platform)" />	
@@ -223,7 +224,7 @@
 						<Label Category="Test" Name="$($(TestName)Short) $(Platform)" Requires="$($(TestName)Short) Tests $(Platform)"/>
 						<Property Name="TestNodes" Value="$(TestNodes);$($(TestName)Short) Tests $(Platform)" If="$(RunAllTests) Or $(Run$(TestName)Tests)"/>
 						
-						<Do If="'$($(TestName)HasAfterSteps)' == 'True'">
+						<Do If="'$(RunAfterSteps)' == 'True' AND '$($(TestName)HasAfterSteps)' == 'True'">
 							<Expand Name="ExpandAfterSteps" TestName="$(TestName)" Platform="$(Platform)" />
 						</Do>
 					</Do>
