Add ability to have AdditionalEditorCompileArgs
#rnx

diff --git a/Engine/Build/EditorPerf.xml b/Engine/Build/EditorPerf.xml
--- a/Engine/Build/EditorPerf.xml
+++ b/Engine/Build/EditorPerf.xml
@@ -15,6 +15,7 @@
 	<Option Name="RunCategory" DefaultValue="Test" Description="Category string for Run Label"/>
 	<Option Name="BaselineCategory" DefaultValue="Test" Description="Category string for Baseline Label"/>
 	<Option Name="SingleAgentCompile" Restrict="true|false" DefaultValue="false" Description="Run all compile steps on same agent"/>
+	<Option Name="AdditionalEditorCompileArgs" DefaultValue="" Description="Additional Compile Arguments for the Editor"/>
 
 	<!-- Projects setup-->
 	<Option Name="Projects" DefaultValue="false" Restrict=".+" Description="The projects to test"/>
@@ -41,6 +42,7 @@
 	<Property Name="EditorCompileArgs" Value="-BuildPlugin=$(ExtraPluginsPlusDelims)" If="'$(ExtraPluginsPlusDelims)' != ''"/>
 	<Property Name="EditorCompileArgs" Value="$(EditorCompileArgs) -nodebuginfo" If="$(NoDebugInfo)"/>
 	<Property Name="EditorCompileArgs" Value="$(EditorCompileArgs) -EnableTSAN -Monolithic" If="$(WithTSAN)"/>
+	<Property Name="EditorCompileArgs" Value="$(EditorCompileArgs) $(AdditionalEditorCompileArgs)" If="'$(AdditionalEditorCompileArgs)' != ''"/>
 
 	<!-- Automation uses its own network output directory -->
 	<Property Name="AutomationNetworkRootDir" Value="\\epicgames.net\root\Builds\Automation\Reports\EditorPerf"/>
