Linux: Fix IsVanilla check by implementing EnumerateLauncherEngineInstallations for Linux

- Also add -BuildId flag to allow released builds to be registered as "launcher" builds.
- This gets us telemetry and the ability to open projects made on the Mac and Windows builds without being prompted to open a copy

diff --git a/Engine/Build/InstalledEngineBuild.xml b/Engine/Build/InstalledEngineBuild.xml
--- a/Engine/Build/InstalledEngineBuild.xml
+++ b/Engine/Build/InstalledEngineBuild.xml
@@ -1084,7 +1084,7 @@
 			<Property Name="ArchitectureCommands" Value="$(ArchitectureCommands) -AndroidArchitectures=$(AndroidArchitectures)" If="'$(WithAndroid)' == true"/>
 
 			<!-- Finalize the build and then tag the text file that this process creates -->
-			<Command Name="FinalizeInstalledBuild" Arguments="-OutputDir=&quot;$(LocalInstalledDir)&quot; -Platforms=$(Platforms) -ContentOnlyPlatforms=$(ContentOnlyPlatforms) $(ArchitectureCommands) -AnalyticsTypeOverride=$(AnalyticsTypeOverride)"/>
+			<Command Name="FinalizeInstalledBuild" Arguments="-OutputDir=&quot;$(LocalInstalledDir)&quot; $(BuildIdentifierArg) -Platforms=$(Platforms) -ContentOnlyPlatforms=$(ContentOnlyPlatforms) $(ArchitectureCommands) -AnalyticsTypeOverride=$(AnalyticsTypeOverride)"/>
 			<Tag Files="$(LocalInstalledDir)/Engine/Build/InstalledBuild.txt" With="#Installed Build Win64 Files"/>
 			
 			<!-- Sanitize all the receipts in the target directory -->
