[EditorPerfTests] Reuse per branch asset registry to make incremental runs representative of user experience

diff --git a/Engine/Build/EditorPerf.xml b/Engine/Build/EditorPerf.xml
--- a/Engine/Build/EditorPerf.xml
+++ b/Engine/Build/EditorPerf.xml
@@ -166,6 +166,9 @@
 	<Property Name="TelemetryConfig" Value="UETelemetryStaging" If="$(IsBuildMachine) And $(IsPreflight) And '$(TelemetryConfig)' == '' And $(RunTraceSummary)"/>
 	<Property Name="TelemetryConfig" Value="UETelemetry" If="$(IsBuildMachine) And '$(TelemetryConfig)' == '' And $(RunTraceSummary)"/>
 
+	<!-- Option to retain per-branch asset registry between runs so that incremental is representative of user experience -->
+	<Property Name="AssetRegistryCache_Folder" Value="$(RootDir)/../../../EditorPerfAssetRegistryCache/$(EscapedBranch)/"/>
+
 	<!-- The main test running macro logic that is expanded with platform and dependency node names -->
 	<Macro Name="RunTest" Arguments="TestNameSuffix;Platform;ProjectName;ProjectTarget;NodeName;Phase;Iteration;Args;EditorArgs">
 		<Property Name="TestName" Value="EditorPerf$(Phase)$(TestNameSuffix)"/>
@@ -185,7 +188,7 @@
 		<Property Name="ReportParams" Value="-WriteTestResultsForHorde -HordeTestDataKey=&quot;$(NodeName)$(PreflightSuffix)&quot;"/>
 		<Property Name="ArtifactName" Value="$(Phase)$(IterationSuffix)$(TestNameSuffix)"/>
 		<Property Name="AllGauntletArgs" Value="-Project=$(ProjectTarget) -Platform=$(Platform) $(GauntletArgs) $(ReportParams) $(LogOutputParams) -artifactname=&quot;$(ArtifactName)&quot;"/>
-		<Property Name="RunArgs" Value="$(Args) -editorargs=&quot;$(EditorArgs)&quot;"/>
+		<Property Name="RunArgs" Value="$(Args) -editorargs=&quot;$(EditorArgs)&quot; -AssetRegistryCacheRootFolder=&quot;$(AssetRegistryCache_Folder)&quot;"/>
 
 		<!-- Do DDC cleanup outside the editor process, handles Cold cleanup without changing stats -->
 		<!-- <Commandlet Project="$(ProjectTarget)" Name="DDCCleanup" Arguments="$(Args) $(CommandletArgs)"/> -->
