Allow passing extra compile arguments for the editor compile

#rnx

diff --git a/Engine/Build/CompileShadersTest.xml b/Engine/Build/CompileShadersTest.xml
--- a/Engine/Build/CompileShadersTest.xml
+++ b/Engine/Build/CompileShadersTest.xml
@@ -13,6 +13,7 @@
 	<Property Name="BuildName" Value="$(EscapedBranch)-CL-$(Change)$(PreflightSuffix)"/>
 	<Option Name="CompileCategory" DefaultValue="Editors" Description="Category string for Compile Label"/>
 	<Option Name="RunCategory" DefaultValue="Test" Description="Category string for Run Label"/>
+	<Option Name="AdditionalEditorCompileArgs" DefaultValue="" Description="Additional Compile Arguments for the Editor"/>
 	
 	<Property Name="AnalyticsArg" Value=""/>
 	<Option Name="SendAnalytics" DefaultValue="false" Description="Whether or not to shader stat analytics."/>
@@ -46,7 +47,7 @@
 			<Node Name="$(CompileNodeName)">
 				<Compile Target="ShaderCompileWorker" Platform="$(EditorPlatform)" Configuration="Development"/>
 				<Compile Target="CrashReportClientEditor" Platform="$(EditorPlatform)" Project="$(Project)" Configuration="Shipping"/>
-				<Compile Target="$(EditorTarget)" Platform="$(EditorPlatform)" Project="$(Project)" Configuration="Development"/>
+				<Compile Target="$(EditorTarget)" Platform="$(EditorPlatform)" Project="$(Project)" Configuration="Development" Arguments="$(AdditionalEditorCompileArgs)"/>
 			</Node>
 
 			<Label Category="$(CompileCategory)" Name="EP Compile $(Project) $(EditorPlatform)" Requires="$(CompileNodeName)"/>
