Add option to force AfterSteps to be True when ran outside of preflights.

diff --git a/Engine/Build/LowLevelTests.xml b/Engine/Build/LowLevelTests.xml
--- a/Engine/Build/LowLevelTests.xml
+++ b/Engine/Build/LowLevelTests.xml
@@ -104,6 +104,11 @@
 	
 	<!-- After Steps -->
 	<Option Name="RunAfterSteps" DefaultValue="false" Description="If set RunAfterSteps is True then we should run aftersteps for this run."/>
+			
+	<!-- If this is NOT a Preflight (with $(PreflightChange) set) OR if the user explictly requested, then Run After Steps -->
+	<Property Name="CheckedRunAfterStep" Value="false"/>
+	<Property Name="CheckedRunAfterStep" Value="True" If="'$(PreflightChange)' == '' OR $(RunAfterSteps)"/>
+
 	<Macro Name="ExpandAfterSteps" Arguments="Platform;TestName">
 		<Node Name="$(TestName) AfterStep $(Platform)">
 			<Expand Name="$(TestName)AfterSteps" Platform="$(Platform)" />	
@@ -252,7 +257,7 @@
 						<Label Category="Test" Name="$($(TestName)Short) $(Platform)" Requires="$($(TestName)Short) Tests $(Platform)"/>
 						<Property Name="TestNodes" Value="$(TestNodes);$($(TestName)Short) Tests $(Platform)" If="$(RunAllTests) Or $(Run$(TestName)Tests)"/>
 						
-						<Do If="'$(RunAfterSteps)' == 'True' AND '$($(TestName)HasAfterSteps)' == 'True'">
+						<Do If="'$(CheckedRunAfterStep)' == 'True' AND '$($(TestName)HasAfterSteps)' == 'True'">
 							<Expand Name="ExpandAfterSteps" TestName="$(TestName)" Platform="$(Platform)" />
 						</Do>
 					</Do>
