Fix vulkan validation layers not showing up in installed engine builds

note: preflight is from this change unshelved in UE5/Main so that I could run a scheduled release build

diff --git a/Engine/Build/InstalledEngineFilters.xml b/Engine/Build/InstalledEngineFilters.xml
--- a/Engine/Build/InstalledEngineFilters.xml
+++ b/Engine/Build/InstalledEngineFilters.xml
@@ -441,6 +441,9 @@
 		<!-- App local dependencies -->
 		Engine/Binaries/ThirdParty/AppLocalDependencies/...
 		
+		<!-- Vulkan Validation Layers -->
+		Engine/Binaries/ThirdParty/Vulkan/...
+
 		<!-- Oculus Libraries & Build Files -->
 		Engine/Source/ThirdParty/Oculus/LibOVRAudio/LibOVRAudio/lib/win64/*.lib
 		Engine/Source/ThirdParty/Oculus/LibOVRAudio/*.tps
@@ -613,6 +616,7 @@
 
 	<Property Name="CopyLinuxFilter">
 		Engine/Binaries/ThirdParty/OpenAL/...
+		Engine/Binaries/ThirdParty/Vulkan/...
 		Engine/Binaries/ThirdParty/ICU/icu4c-53_1/Unix/x86_64-unknown-linux-gnu/*.so
 		Engine/Source/ThirdParty/Unix/LibCxx/...
 		Engine/Source/ThirdParty/Unix/HaveLinuxDependencies
