Include zenserver, zen, and crashpad_handler executables when producing installed engine builds for UE.

diff --git a/Engine/Build/InstalledEngineFilters.xml b/Engine/Build/InstalledEngineFilters.xml
--- a/Engine/Build/InstalledEngineFilters.xml
+++ b/Engine/Build/InstalledEngineFilters.xml
@@ -312,6 +312,11 @@
 		
 		<!-- UnrealTraceServer is the tool that records Unreal Insights trace files -->
 		Engine/Binaries/Win64/UnrealTraceServer.exe
+		
+		<!-- zenserver is the tool for Unreal Zen Storage services like Local DDC and optionally cooked output storage -->
+		Engine/Binaries/Win64/zenserver.exe
+		Engine/Binaries/Win64/zen.exe
+		Engine/Binaries/Win64/crashpad_handler.exe
 
 		<!-- Used by Switchboard application, and by iOS/tvOS build automation -->
 		Engine/Extras/ThirdPartyNotUE/cwrsync/...
@@ -364,6 +369,11 @@
 		
 		<!-- UnrealTraceServer is the tool that records Unreal Insights trace files -->
 		Engine/Binaries/Mac/UnrealTraceServer
+		
+		<!-- zenserver is the tool for Unreal Zen Storage services like Local DDC and optionally cooked output storage -->
+		Engine/Binaries/Mac/zenserver
+		Engine/Binaries/Mac/zen
+		Engine/Binaries/Mac/crashpad_handler
 	</Property>
 
 	<Property Name="CopyEditorExceptionsMac" Value="$(CopyEditorExceptions)">
@@ -416,6 +426,10 @@
 		
 		<!-- UnrealTraceServer is the tool that records Unreal Insights trace files -->
 		Engine/Binaries/Linux/UnrealTraceServer
+		
+		<!-- zenserver is the tool for Unreal Zen Storage services like Local DDC and optionally cooked output storage -->
+		Engine/Binaries/Linux/zenserver
+		Engine/Binaries/Linux/zen
 	</Property>
 	<Property Name="CopyEditorExceptionsLinux" Value="$(CopyEditorExceptions)">
 		<!-- Exclude Windows/Mac binaries on Linux here -->
