Add LiveLinkHub to the UE installed build

diff --git a/Engine/Build/InstalledEngineBuild.xml b/Engine/Build/InstalledEngineBuild.xml
--- a/Engine/Build/InstalledEngineBuild.xml
+++ b/Engine/Build/InstalledEngineBuild.xml
@@ -483,6 +483,7 @@
 			<Compile Target="SwitchboardListenerHelper" Configuration="Development" Platform="Win64" Tag="#Build Tools Win64"/>
 			<Compile Target="UnrealObjectPtrTool" Configuration="Development" Platform="Win64" Tag="#Build Tools Win64"/>
 			<Compile Target="ZenDashboard" Configuration="Development" Platform="Win64" Tag="#Build Tools Win64"/>
+			<Compile Target="LiveLinkHub" Configuration="Shipping" Platform="Win64" Arguments="-Monolithic" Tag="#Build Tools Win64"/>
 
 			<Tag Files="Engine/Binaries/Win64/XGEControlWorker.exe" With="#Build Tools Win64"/>
 
@@ -636,6 +637,7 @@
 			<Compile Target="UnrealAtoS" Configuration="Development" Platform="Mac" Tag="#Build Tools Mac" Arguments="-EnableDSYM"/>
 			<Compile Target="UnrealObjectPtrTool" Configuration="Development" Platform="Mac" Tag="#Build Tools Mac" Arguments="-EnableDSYM"/>
 			<Compile Target="ZenDashboard" Configuration="Development" Platform="Mac" Tag="#Build Tools Mac" Arguments="-EnableDSYM"/>
+			<Compile Target="LiveLinkHub" Configuration="Shipping" Platform="Mac" Tag="#Build Tools Mac" Arguments="-EnableDSYM -Monolithic"/>
 		</Node>
 	</Agent>
 
@@ -657,6 +659,7 @@
 			<Compile Target="SwitchboardListenerHelper" Configuration="Development" Platform="Linux" Tag="#Build Tools LinuxNative"/>
 			<Compile Target="UnrealObjectPtrTool" Configuration="Development" Platform="Linux" Tag="#Build Tools LinuxNative"/>
 			<Compile Target="ZenDashboard" Configuration="Development" Platform="Linux" Tag="#Build Tools LinuxNative"/>
+			<Compile Target="LiveLinkHub" Configuration="Shipping" Platform="Linux" Tag="#Build Tools LinuxNative" Arguments="-Monolithic"/>
 		</Node>
 
 		<!-- Compile the C# tools. Copies files to a temporary directory first, to strip out all confidential folders and avoid sharing violations with the running version of UAT. -->
