Disable livelinkhub dSYM generation to avoid hitting dsym size limit

diff --git a/Engine/Build/InstalledEngineBuild.xml b/Engine/Build/InstalledEngineBuild.xml
--- a/Engine/Build/InstalledEngineBuild.xml
+++ b/Engine/Build/InstalledEngineBuild.xml
@@ -637,7 +637,7 @@
 			<Compile Target="UnrealAtoS" Configuration="Development" Platform="Mac" Tag="#Build Tools Mac" Arguments="-EnableDSYM"/>
 			<Compile Target="UnrealObjectPtrTool" Configuration="Development" Platform="Mac" Tag="#Build Tools Mac" Arguments="-EnableDSYM"/>
 			<Compile Target="ZenDashboard" Configuration="Development" Platform="Mac" Tag="#Build Tools Mac" Arguments="-EnableDSYM"/>
-			<Compile Target="LiveLinkHub" Configuration="Shipping" Platform="Mac" Tag="#Build Tools Mac" Arguments="-EnableDSYM -Monolithic"/>
+			<Compile Target="LiveLinkHub" Configuration="Shipping" Platform="Mac" Tag="#Build Tools Mac" Arguments="-Monolithic"/>
 		</Node>
 	</Agent>
 
