Remove References to internal Horde server deployment and network path

#rnx

diff --git a/Engine/Build/EditorPerf.xml b/Engine/Build/EditorPerf.xml
--- a/Engine/Build/EditorPerf.xml
+++ b/Engine/Build/EditorPerf.xml
@@ -3,6 +3,7 @@
 
 	<!-- Horde Options -->	
 	<EnvVar Name="UE_HORDE_JOBID"/>
+	<EnvVar Name="UE_HORDE_STEPID"/>
 	<Option Name="Horde" Restrict="true|false" DefaultValue="false" Description="Whether the build is running under Horde"/>
 	<Option Name="PreflightChange" DefaultValue="" Description="Preflight CL number if preflight, empty otherwise"/>
 	<Option Name="SetPreflightFlag" DefaultValue="true" Description="Whether to set the IsPreflight flag to true for preflights. Use with caution: this will cause things like build versions to be set to their non-preflight form."/>
@@ -44,10 +45,12 @@
 	<Property Name="EditorCompileArgs" Value="$(EditorCompileArgs) -EnableTSAN -Monolithic" If="$(WithTSAN)"/>
 	<Property Name="EditorCompileArgs" Value="$(EditorCompileArgs) $(AdditionalEditorCompileArgs)" If="'$(AdditionalEditorCompileArgs)' != ''"/>
 
+	<!-- Get shared properties -->
+	<Include Script="Graph/Tasks/Inc/SharedProperties.xml" />
+
 	<!-- Automation uses its own network output directory -->
-	<Property Name="AutomationNetworkRootDir" Value="\\epicgames.net\root\Builds\Automation\Reports\EditorPerf"/>
-	<Property Name="AutomationNetworkRootDir" Value="/Volumes/Builds/Automation/Reports/EditorPerf" If="'$(HostPlatform)' != 'Win64'"/>
-	<Property Name="AutomationNetworkRootDir" Value="/mnt/Builds/Automation/Reports/EditorPerf" If="'$(HostPlatform)' == 'Linux'"/>
+	<Property Name="AutomationNetworkRootDir" Value="$(NetworkReportRoot)\EditorPerf"/>
+	<Property Name="AutomationNetworkRootDir" Value="$(NetworkReportRoot)/EditorPerf" If="'$(HostPlatform)' != 'Win64'"/>
 	<Property Name="AutomationLocalOutputDir" Value="$(RootDir)/LocalBuilds/Reports"/>
 
 	<!-- Outputs for this run keeping slashes flipped the right way because it shows up on the webapp this way -->
@@ -168,6 +171,8 @@
 	<Option Name="TelemetryConfig" DefaultValue="" Description="The configuration to use when publishing telemetry data to a database"/>
 	<Property Name="TelemetryConfig" Value="UETelemetryStaging" If="$(IsBuildMachine) And $(IsPreflight) And '$(TelemetryConfig)' == '' And $(RunTraceSummary)"/>
 	<Property Name="TelemetryConfig" Value="UETelemetry" If="$(IsBuildMachine) And '$(TelemetryConfig)' == '' And $(RunTraceSummary)"/>
+	<Option Name="ExtraTelemetryArgs" DefaultValue="" Description="Extra arguments passed to the publish telemetry command"/>
+	<Property Name="ExtraTelemetryArgs" Value="$(ExtraTelemetryArgs) -JobLink=&quot;$(HordeDashboardURL)/job/$(UE_HORDE_JOBID)?step=$(UE_HORDE_STEPID)&quot;" If="'$(HordeDashboardURL)' != ''"/>
 
 	<!-- The main test running macro logic that is expanded with platform and dependency node names -->
 	<Macro Name="RunTest" Arguments="TestNameSuffix;Platform;ProjectName;ProjectTarget;NodeName;Phase;Iteration;Args;EditorArgs">
@@ -205,7 +210,7 @@
 		
 		<Do If="'$(TelemetryConfig)' != ''">
 			<!-- Publish telemetry to the db -->
-			<Command Name="PublishUnrealAutomationTelemetry" Arguments="-CSVFile=&quot;$(TelemetryFilePath)&quot; -project=$(Project) -changelist=$(Change) -branch=$(Branch) -TelemetryConfig=$(TelemetryConfig)"/>
+			<Command Name="PublishUnrealAutomationTelemetry" Arguments="-CSVFile=&quot;$(TelemetryFilePath)&quot; -project=$(Project) -changelist=$(Change) -branch=$(Branch) -TelemetryConfig=$(TelemetryConfig) $(ExtraTelemetryArgs)"/>
 		</Do>		
 	</Macro>
 	
diff --git a/Engine/Build/Graph/Tasks/BuildAndTestProject.xml b/Engine/Build/Graph/Tasks/BuildAndTestProject.xml
--- a/Engine/Build/Graph/Tasks/BuildAndTestProject.xml
+++ b/Engine/Build/Graph/Tasks/BuildAndTestProject.xml
@@ -114,7 +114,7 @@
 	<Option Name="AnnotationsPlatforms" DefaultValue="" Description="List of platforms that require annotations"/>
 	<Option Name="AnnotationsTests" DefaultValue="" Description="List of tests that require annotations"/>
 
-	<!-- Some common properties -->
+	<!-- Default common properties -->
 	<Include Script="Inc/CommonProperties.xml" />
 
 	<!-- Common settings but can be overidden -->
diff --git a/Engine/Build/Graph/Tasks/Inc/CommonProperties.xml b/Engine/Build/Graph/Tasks/Inc/CommonProperties.xml
--- a/Engine/Build/Graph/Tasks/Inc/CommonProperties.xml
+++ b/Engine/Build/Graph/Tasks/Inc/CommonProperties.xml
@@ -21,4 +21,6 @@
 	<Property Name="BuildNamePath" Value="$(ProjectName)-$(CL-String)"/>
 	<Property Name="BuildVersion" Value="$(ProjectName)-CL-$(Change)"/>
 
+	<Include Script="SharedProperties.xml" />
+
 </BuildGraph>
diff --git a/Engine/Build/Graph/Tasks/Inc/GauntletSettings.xml b/Engine/Build/Graph/Tasks/Inc/GauntletSettings.xml
--- a/Engine/Build/Graph/Tasks/Inc/GauntletSettings.xml
+++ b/Engine/Build/Graph/Tasks/Inc/GauntletSettings.xml
@@ -17,6 +17,7 @@
 
 	<!-- How long to give Gauntlet before a test is considered timed out (some tests will time out earlier if they detect inactivity)-->
 	<Option Name="GauntletTimeout" DefaultValue="3600" Description="Max running time of a test in Gauntlet"/>
+	<Option Name="GauntletDeviceService" DefaultValue="" Description="Gauntlet Device Service to use for device reservation when running tests"/>
 	<Option Name="GauntletDevicePool" DefaultValue="UE4" Description="Gauntlet Device Pool to use to run target platform tests"/>
 	<Option Name="AlternateGauntletDevicePool" DefaultValue="UE5" Description="Alternate Gauntlet Device Pool to use to run target platform tests"/>
 	<Option Name="AlternateDevicePoolPlatforms" DefaultValue="XSX;PS5" Description="Platforms that need to use the Alternate Gauntlet Device Pool"/>
@@ -32,11 +33,13 @@
 	<Property Name="GauntletCommonArgs" Value="$(GauntletCommonArgs) -branch=&quot;$(Branch)&quot; -changelist=&quot;$(Change)&quot;" />
 
 	<Property Name="GauntletDeviceArgs" Value="" />
-	<Property Name="GauntletDeviceArgs" Value="-deviceurl=https://horde.devtools.epicgames.com/ -reboot -cleardevices -sandbox=&quot;$(ProjectName)&quot;" If="$(IsBuildMachine)" />
+	<Property Name="GauntletDeviceArgs" Value="-reboot -cleardevices -sandbox=&quot;$(ProjectName)&quot;" If="$(IsBuildMachine)" />
 	<Property Name="GauntletDeviceArgs" Value="$(GauntletDeviceArgs) -VerifyLogin" If="$(EnableDeviceLoginVerification)"/>
 	<Property Name="GauntletDeviceAltArgs" Value="$(GauntletDeviceArgs)" />
-	<Property Name="GauntletDeviceArgs" Value="$(GauntletDeviceArgs) -devicepool=&quot;$(GauntletDevicePool)&quot;" If="$(IsBuildMachine)"/>
-	<Property Name="GauntletDeviceAltArgs" Value="$(GauntletDeviceAltArgs) -devicepool=&quot;$(AlternateGauntletDevicePool)&quot;" If="$(IsBuildMachine)"/>
+	<Do If="$(IsBuildMachine) And '$(GauntletDeviceService)' != ''">
+		<Property Name="GauntletDeviceArgs" Value="$(GauntletDeviceArgs) -deviceurl=$(GauntletDeviceService) -devicepool=&quot;$(GauntletDevicePool)&quot;"/>
+		<Property Name="GauntletDeviceAltArgs" Value="$(GauntletDeviceAltArgs) -deviceurl=$(GauntletDeviceService) -devicepool=&quot;$(AlternateGauntletDevicePool)&quot;"/>
+	</Do>
 
 	<Property Name="GauntletTargetTestArgs" Value="$(GauntletCommonArgs) $(GauntletDeviceArgs) $(ExtraGauntletTargetTestArgs)" />
 	<Property Name="GauntletTargetTestAltArgs" Value="$(GauntletCommonArgs) $(GauntletDeviceAltArgs) $(ExtraGauntletTargetTestArgs)" />
diff --git a/Engine/Build/Graph/Tasks/Inc/PathSettings.xml b/Engine/Build/Graph/Tasks/Inc/PathSettings.xml
--- a/Engine/Build/Graph/Tasks/Inc/PathSettings.xml
+++ b/Engine/Build/Graph/Tasks/Inc/PathSettings.xml
@@ -16,6 +16,9 @@
 		InProperty: ProjectName
 		InProperty: ProjectPath
 		InProperty: EscapedBranch
+		InProperty: NetworkTempRoot
+		InProperty: NetworkPublishRoot
+		InProperty: NetworkReportRoot
 
 		OutOption: PublishBuilds - Publish output from builds to a network share
 		OutOption: PublishHashedStagingData - Publish staged content with hashes for quick syncing
@@ -46,21 +49,6 @@
 	<Property Name="LocalInstallDirectory" Value="$(ProjectPath)/Saved/LocalInstall"/>
 	<Property Name="AutomationReportOutputDirectory" Value=""/>
 
-	<!-- Scratch network path that can be used for tests if we aren't publishing. Assumes there's a cleanup script.. -->
-	<Property Name="NetworkTempRoot" Value="/Volumes/Builds/UE5/TempStorage"/>
-	<Property Name="NetworkTempRoot" Value="/mnt/Builds/UE5/TempStorage" If="'$(HostPlatform)' == 'Linux'"/>
-	<Property Name="NetworkTempRoot" Value="P:/Builds/UE5/TempStorage" If="'$(HostPlatform)' == 'Win64'"/>
-
-	<!-- Network path that can be used for builds that are published -->
-	<Property Name="NetworkPublishRoot" Value="/Volumes/Builds/UE5/PackagedBuilds"/>
-	<Property Name="NetworkPublishRoot" Value="/mnt/Builds/UE5/PackagedBuilds" If="'$(HostPlatform)' == 'Linux'"/>
-	<Property Name="NetworkPublishRoot" Value="P:/Builds/UE5/PackagedBuilds" If="'$(HostPlatform)' == 'Win64'"/>
-
-	<!-- Network path that can be used for storing reports and logs -->
-	<Property Name="NetworkReportRoot" Value="/Volumes/Builds/Automation/Reports"/>
-	<Property Name="NetworkReportRoot" Value="/mnt/Builds/Automation/Reports" If="'$(HostPlatform)' == 'Linux'"/>
-	<Property Name="NetworkReportRoot" Value="P:/Builds/Automation/Reports" If="'$(HostPlatform)' == 'Win64'"/>
-
 	<!-- Allow scripts to override the network locations-->
 	<Do If="'$(NetworkTempRootOverride)' != ''" >
 		<Property Name="NetworkTempRoot" Value="$(NetworkTempRootOverride)"/>
diff --git a/Engine/Build/Graph/Tasks/Inc/SharedProperties.xml b/Engine/Build/Graph/Tasks/Inc/SharedProperties.xml
new file mode 100644
--- /dev/null
+++ b/Engine/Build/Graph/Tasks/Inc/SharedProperties.xml
@@ -0,0 +1,20 @@
+<?xml version='1.0' ?>
+<BuildGraph xmlns="http://www.epicgames.com/BuildGraph" xmlns:xsi="http://www.w3.org/2001/XMLSchema-instance" xsi:schemaLocation="http://www.epicgames.com/BuildGraph ../Schema.xsd" >
+
+	<!-- Option include files -->
+	<Option Name="SharedPropertiesFile" DefaultValue="Engine/Restricted/NotForLicensees/Build/Graph/Tasks/Inc/SharedProperties.xml" Description="Shared Properties file to use (relative from root view)." />
+	<!-- Scratch network path that can be used for tests if we aren't publishing. Assumes there's a cleanup script.. -->
+	<Property Name="NetworkTempRoot" Value=""/>
+	<!-- Network path that can be used for builds that are published -->
+	<Property Name="NetworkPublishRoot" Value=""/>
+	<!-- Network path that can be used for storing reports and logs -->
+	<Property Name="NetworkReportRoot" Value=""/>
+	<!-- URL to Horde server -->
+	<Property Name="HordeDashboardURL" Value=""/>
+	<Option Name="GauntletDeviceService" DefaultValue="" Description="Gauntlet Device Service to use for device reservation when running tests"/>
+	<!-- Get custom shared properties -->
+	<Do If="Exists('$(RootDir)/$(SharedPropertiesFile)')">
+		<Include Script="$(RootDir)/$(SharedPropertiesFile)"/>
+	</Do>
+
+</BuildGraph>
diff --git a/Engine/Build/LowLevelTests.xml b/Engine/Build/LowLevelTests.xml
--- a/Engine/Build/LowLevelTests.xml
+++ b/Engine/Build/LowLevelTests.xml
@@ -84,7 +84,6 @@
 		<Property Name="TestsAllSupportedPlatforms" Value="false" If="$(TestPlatform$(Platform))"/>
 	</ForEach>
 	
-	<Property Name="GauntletDeviceService" Value="https://horde.devtools.epicgames.com" />
 	<Option Name="TestArgs" DefaultValue="" Description="Extra test arguments."/>
 
 	<Option Name="GauntletDevicePool" DefaultValue="UE4" Description="Gauntlet Device Pool to use to run target platform tests"/>
@@ -95,6 +94,9 @@
 	<Option Name="AttachToDebugger" DefaultValue="false" Description="Wait for debugger to attach."/>
 	<Option Name="EnableDeviceLoginVerification" DefaultValue="$(IsBuildMachine)" Description="Verify a user is logged in on Devices that supports it, try to sign in if not."/>
 
+	<!-- Get shared properties -->
+	<Include Script="Graph/Tasks/Inc/SharedProperties.xml" />
+
 	<Property Name="TestNodes" Value=""/>
 	<Property Name="AfterTestNodes" Value=""/>
 	<Property Name="Catch2BuildNodes" Value=""/>
@@ -134,7 +136,8 @@
 		<Property Name="PlatformStageDirName" Value="Windows" If="'$(Platform)' == 'Win64'"/>
 		<Property Name="DevicePool" Value="$(GauntletDevicePool)" />
 		<Property Name="DevicePool" Value="$(AlternateGauntletDevicePool)" If="ContainsItem('$(AlternateDevicePoolPlatforms)', '$(Platform)', ';')" />
-		<Property Name="TestArgsExt" Value="$(TestArgs) -deviceurl=$(GauntletDeviceService) -devicepool=$(DevicePool)"/>
+		<Property Name="TestArgsExt" Value="$(TestArgs)"/>
+		<Property Name="TestArgsExt" Value="$(TestArgsExt) -deviceurl=$(GauntletDeviceService) -devicepool=$(DevicePool)" If="'$(GauntletDeviceService)' != ''"/>
 		<Property Name="TestModeSuffix" Value=""/>
 		<Property Name="TestTags" Value=""/>
 		<Property Name="TestTags" Value="$($(Platform)$(TestName)Tags)" If="'$($(Platform)$(TestName)Tags)' != ''"/>
diff --git a/Engine/Source/Programs/AutomationTool/Gauntlet/Framework/Utils/Gauntlet.HordeReport.cs b/Engine/Source/Programs/AutomationTool/Gauntlet/Framework/Utils/Gauntlet.HordeReport.cs
--- a/Engine/Source/Programs/AutomationTool/Gauntlet/Framework/Utils/Gauntlet.HordeReport.cs
+++ b/Engine/Source/Programs/AutomationTool/Gauntlet/Framework/Utils/Gauntlet.HordeReport.cs
@@ -35,53 +35,6 @@ namespace Gauntlet
 				return Path.GetFullPath(Environment.GetEnvironmentVariable("UE_ARTIFACTS_DIR") ?? CommandUtils.CmdEnv.LogFolder);
 			}
 		}
-		/// <summary>
-		/// Is Environement set by Horde Agent 
-		/// </summary>
-		public static bool IsUnderHordeAgent
-		{
-			get
-			{
-				return !string.IsNullOrEmpty(Environment.GetEnvironmentVariable("UE_TESTDATA_DIR"));
-			}
-		}
-		/// <summary>
-		/// Default Job Horde link if running under Horde agent, otherwise return an empty string
-		/// </summary>
-		public static string DefaultHordeJobLink
-		{
-			get
-			{
-				string HordeJobId = Environment.GetEnvironmentVariable("UE_HORDE_JOBID");
-				if (!string.IsNullOrEmpty(HordeJobId))
-				{
-					return string.Format("https://horde.devtools.epicgames.com/job/{0}", HordeJobId);
-				}
-
-				return string.Empty;
-			}
-		}
-		/// <summary>
-		/// Default Job Step Horde link if running under Horde agent, otherwise return an empty string
-		/// </summary>
-		public static string DefaultHordeJobStepLink
-		{
-			get
-			{
-				string HordeJobStepId = Environment.GetEnvironmentVariable("UE_HORDE_STEPID");
-				string HordeJobLink = DefaultHordeJobLink;
-				if (!string.IsNullOrEmpty(HordeJobStepId))
-				{
-					return string.Format("{0}?step={1}", HordeJobLink, HordeJobStepId);
-				}
-				else if(!string.IsNullOrEmpty(HordeJobLink))
-				{
-					return HordeJobLink;
-				}
-
-				return string.Empty;
-			}
-		}		
 
 		public abstract class BaseHordeReport : BaseTestReport
 		{
diff --git a/Engine/Source/Programs/AutomationTool/Gauntlet/Unreal/Utils/Gauntlet.UnrealTelemetry.cs b/Engine/Source/Programs/AutomationTool/Gauntlet/Unreal/Utils/Gauntlet.UnrealTelemetry.cs
--- a/Engine/Source/Programs/AutomationTool/Gauntlet/Unreal/Utils/Gauntlet.UnrealTelemetry.cs
+++ b/Engine/Source/Programs/AutomationTool/Gauntlet/Unreal/Utils/Gauntlet.UnrealTelemetry.cs
@@ -183,7 +183,7 @@ namespace Gauntlet
 		[Help("Branch", "Target Branch name. Default: Unknown.")]
 		[Help("Changelist", "Target Changelist number. Default: 0.")]
 		[Help("Configuration", "Target Configuration name. Default: Development.")]
-		[Help("JobLink", "Http Link to Build Job. Default: %UE_HORDE_JOBID%")]
+		[Help("JobLink", "Http Link to Build Job.")]
 
 		public override ExitCode Execute()
 		{
@@ -272,10 +272,6 @@ namespace Gauntlet
 			Context.SetProperty("Changelist", ChangelistString);
 			Context.SetProperty("ChangelistDateTime", GetChangelistDateTime(int.Parse(ChangelistString)));
 
-			if (string.IsNullOrEmpty(JobLink))
-			{
-				JobLink = HordeReport.DefaultHordeJobStepLink;
-			}
 			Context.SetProperty("JobLink", JobLink);
 
 			// Create a report
