Split CRC analytics definition into a development and retail variant.

We want to be able to get telemetry event from binary builds of the engine, but only in development time (e.g. CRCE), not for the packaged game.

diff --git a/Engine/Build/InstalledEngineBuild.xml b/Engine/Build/InstalledEngineBuild.xml
--- a/Engine/Build/InstalledEngineBuild.xml
+++ b/Engine/Build/InstalledEngineBuild.xml
@@ -146,10 +146,17 @@
 	<!-- Include the script that lists files and patterns used when creating an installed build -->
 	<Include Script="InstalledEngineFilters.xml"/>
 
-	<Option Name="CrashReporterAPIURL" DefaultValue="" Description="The URL to use to talk to the CrashReporterClient API."/>
-	<Option Name="CrashReporterAPIKey" DefaultValue="" Description="The API key to use to talk to the CrashReporterClient API."/>
+	<!-- Crash reporter telemetry settings for the customer facing crash reporter -->
+	<Option Name="CrashReporterAPIURL" DefaultValue="" Description="The URL to use to talk to the CrashReporterClient API. Used for crashes from packed game."/>
+	<Option Name="CrashReporterAPIKey" DefaultValue="" Description="The API key to use to talk to the CrashReporterClient API. Used for crashes from packaged game."/>
 	<Property Name="CrashReporterCompileArgs" Value=""/>
 	<Property Name="CrashReporterCompileArgs" Value="-define:CRC_TELEMETRY_URL=\&quot;$(CrashReporterAPIURL)\&quot; -define:CRC_TELEMETRY_KEY_DEV=\&quot;$(CrashReporterAPIKey)\&quot; -define:CRC_TELEMETRY_KEY_RELEASE=\&quot;$(CrashReporterAPIKey)\&quot; -OverrideBuildEnvironment" If="'$(CrashReporterAPIURL)' != '' And '$(CrashReporterAPIKey)' != ''" />
+	
+	<!-- Crash reporter telemetry setting for development time crash reporter -->
+	<Option Name="CrashReporterDevAPIURL" DefaultValue="" Description="The URL to use to talk to the CrashReporterClient API. Used for development time crashes."/>
+	<Option Name="CrashReporterDevAPIKey" DefaultValue="" Description="The API key to use to talk to the CrashReporterClient API. Used for development time crashes."/>
+	<Property Name="CrashReporterCompileArgsDev" Value=""/>
+	<Property Name="CrashReporterCompileArgsDev" Value="-define:CRC_TELEMETRY_URL=\&quot;$(CrashReporterDevAPIURL)\&quot; -define:CRC_TELEMETRY_KEY_DEV=\&quot;$(CrashReporterDevAPIKey)\&quot; -define:CRC_TELEMETRY_KEY_RELEASE=\&quot;$(CrashReporterDevAPIKey)\&quot; -OverrideBuildEnvironment" If="'$(CrashReporterDevAPIURL)' != '' And '$(CrashReporterDevAPIKey)' != ''" />
 
 	<Property Name="MacArchitectureArg" Value="-architecture=arm64+x64" />
 
@@ -159,8 +166,8 @@
 			<SetVersion Change="$(Change)" Branch="$(EscapedBranch)" If="$(IsBuildMachine)"/>
 		</Node>
 		<Node Name="Compile UnrealEditor Win64" Requires="Update Version Files" Produces="#UnrealEditor Win64;#UnrealEditor Win64 Unstripped;#UnrealEditor Win64 Stripped;#UnrealEditor Win64 Unsigned;#UnrealEditor Win64 Signed;#UnrealEditor Win64 WithExceptions">
-			<Compile Target="UnrealEditor" Platform="Win64" Configuration="DebugGame" Tag="#UnrealEditor Win64" Arguments="-precompile -allmodules $(CrashReporterCompileArgs)"/>
-			<Compile Target="UnrealEditor" Platform="Win64" Configuration="Development" Tag="#UnrealEditor Win64" Arguments="-precompile -allmodules $(CrashReporterCompileArgs)"/>
+			<Compile Target="UnrealEditor" Platform="Win64" Configuration="DebugGame" Tag="#UnrealEditor Win64" Arguments="-precompile -allmodules"/>
+			<Compile Target="UnrealEditor" Platform="Win64" Configuration="Development" Tag="#UnrealEditor Win64" Arguments="-precompile -allmodules"/>
 
 			<Do If="$(EmbedSrcSrvInfo)">
 				<!-- Embed source info into the PDB files. Should be done from this machine to ensure that paths are correct. -->
@@ -199,8 +206,8 @@
 	<Agent Name="Editor Mac" Type="Mac_Licensee">
 		<Node Name="Compile UnrealEditor Mac" Requires="Update Version Files" Produces="#UnrealEditor Mac;#UnrealEditor Mac Unstripped;#UnrealEditor Mac Stripped;#UnrealEditor Mac Unsigned;#UnrealEditor Mac Signed;#UnrealEditor Mac WithExceptions">
 			<Property Name="AllModulesArg" Value="-allmodules"/>
-			<Compile Target="UnrealEditor" Platform="Mac" Configuration="DebugGame" Tag="#UnrealEditor Mac" Arguments="-precompile -EnableDSYM $(AllModulesArg) $(CrashReporterCompileArgs) $(MacArchitectureArg)"/>
-			<Compile Target="UnrealEditor" Platform="Mac" Configuration="Development" Tag="#UnrealEditor Mac" Arguments="-precompile -EnableDSYM $(AllModulesArg) $(CrashReporterCompileArgs) $(MacArchitectureArg)"/>
+			<Compile Target="UnrealEditor" Platform="Mac" Configuration="DebugGame" Tag="#UnrealEditor Mac" Arguments="-precompile -EnableDSYM $(AllModulesArg) $(MacArchitectureArg)"/>
+			<Compile Target="UnrealEditor" Platform="Mac" Configuration="Development" Tag="#UnrealEditor Mac" Arguments="-precompile -EnableDSYM $(AllModulesArg) $(MacArchitectureArg)"/>
 
 			<Property Name="UnrealEditorExceptions" Value=""/>
 			<Property Name="UnrealEditorExceptions" Value="$(PluginsExceptions)" If="$(CompileDatasmithPlugins)"/>
@@ -231,8 +238,8 @@
 	<!-- Base editor and tools for compile on Linux -->
 	<Agent Name="Editor Linux" Type="Linux_Licensee">
 		<Node Name="Compile UnrealEditor Linux" Requires="Update Version Files" Produces="#UnrealEditor Linux;#UnrealEditor Linux Unstripped;#UnrealEditor Linux Stripped;#UnrealEditor Linux WithExceptions">
-			<Compile Target="UnrealEditor" Platform="Linux" Configuration="DebugGame" Tag="#UnrealEditor Linux" Arguments="-precompile -allmodules $(CrashReporterCompileArgs)"/>
-			<Compile Target="UnrealEditor" Platform="Linux" Configuration="Development" Tag="#UnrealEditor Linux" Arguments="-precompile -allmodules $(CrashReporterCompileArgs)"/>
+			<Compile Target="UnrealEditor" Platform="Linux" Configuration="DebugGame" Tag="#UnrealEditor Linux" Arguments="-precompile -allmodules"/>
+			<Compile Target="UnrealEditor" Platform="Linux" Configuration="Development" Tag="#UnrealEditor Linux" Arguments="-precompile -allmodules"/>
 
 			<Tag Files="#UnrealEditor Linux" Filter="$(LinuxStripFilter)" With="#UnrealEditor Linux Unstripped"/>
 			<Tag Files="#UnrealEditor Linux" Except="$(PluginsExceptions)" With="#UnrealEditor Linux WithExceptions"/>
@@ -464,7 +471,7 @@
 	<Agent Name="Tools Group Win64" Type="Win64_Licensee">
 		<Node Name="Build Tools Win64" Requires="Update Version Files;#NotForLicensee Build Tools Win64" Produces="#Tools Win64 Complete">
 			<Compile Target="CrashReportClient" Configuration="Shipping" Platform="Win64" Tag="#Build Tools Win64" Arguments="$(CrashReporterCompileArgs)"/>
-			<Compile Target="CrashReportClientEditor" Configuration="Shipping" Platform="Win64" Tag="#Build Tools Win64" Arguments="$(CrashReporterCompileArgs)"/>
+			<Compile Target="CrashReportClientEditor" Configuration="Shipping" Platform="Win64" Tag="#Build Tools Win64" Arguments="$(CrashReporterCompileArgsDev)"/>
 			<Tag Files="#NotForLicensee Build Tools Win64" With="#Build Tools Win64"/>
 			<Compile Target="ShaderCompileWorker" Configuration="Development" Platform="Win64" Tag="#Build Tools Win64"/>
 			<Compile Target="EpicWebHelper" Configuration="Shipping" Platform="Win64" Tag="#Build Tools Win64"/>
@@ -550,7 +557,7 @@
 		</Node>
 		<Node Name="Build Tools Linux" Requires="Update Version Files">
 			<Compile Target="CrashReportClient" Platform="Linux" Configuration="Shipping" Tag="#Build Tools Linux" Arguments="$(CrashReporterCompileArgs)"/>
-			<Compile Target="CrashReportClientEditor" Platform="Linux" Configuration="Shipping" Tag="#Build Tools Linux" Arguments="$(CrashReporterCompileArgs)"/>
+			<Compile Target="CrashReportClientEditor" Platform="Linux" Configuration="Shipping" Tag="#Build Tools Linux" Arguments="$(CrashReporterCompileArgsDev)"/>
 			<Compile Target="ShaderCompileWorker" Configuration="Development" Platform="Linux" Tag="#Build Tools Linux"/>
 			<Compile Target="EpicWebHelper" Configuration="Shipping" Platform="Linux" Tag="#Build Tools Linux"/>
 			<Compile Target="UnrealInsights" Configuration="Development" Platform="Linux" Tag="#Build Tools Linux"/>
@@ -620,7 +627,7 @@
 	
 		<Node Name="Build Tools Mac" Requires="Update Version Files">
 			<Compile Target="CrashReportClient" Configuration="Shipping" Platform="Mac" Tag="#Build Tools Mac" Arguments="-EnableDSYM $(CrashReporterCompileArgs)"/>
-			<Compile Target="CrashReportClientEditor" Configuration="Shipping" Platform="Mac" Tag="#Build Tools Mac" Arguments="-EnableDSYM $(CrashReporterCompileArgs)"/>
+			<Compile Target="CrashReportClientEditor" Configuration="Shipping" Platform="Mac" Tag="#Build Tools Mac" Arguments="-EnableDSYM $(CrashReporterCompileArgsDev)"/>
 			<Compile Target="ShaderCompileWorker" Configuration="Development" Platform="Mac" Tag="#Build Tools Mac" Arguments="-EnableDSYM"/>
 			<Compile Target="EpicWebHelper" Configuration="Shipping" Platform="Mac" Tag="#Build Tools Mac" Arguments="-EnableDSYM"/>
 			<Compile Target="UnrealInsights" Configuration="Development" Platform="Mac" Tag="#Build Tools Mac" Arguments="-EnableDSYM"/>
@@ -644,7 +651,7 @@
 	<Agent Name="Tools Group LinuxNative" Type="Linux_Licensee">
 		<Node Name="Build Tools LinuxNative" Requires="Update Version Files">
 			<Compile Target="CrashReportClient" Platform="Linux" Configuration="Shipping" Tag="#Build Tools LinuxNative" Arguments="$(CrashReporterCompileArgs)"/>
-			<Compile Target="CrashReportClientEditor" Platform="Linux" Configuration="Shipping"  Tag="#Build Tools LinuxNative" Arguments="$(CrashReporterCompileArgs)"/>
+			<Compile Target="CrashReportClientEditor" Platform="Linux" Configuration="Shipping"  Tag="#Build Tools LinuxNative" Arguments="$(CrashReporterCompileArgsDev)"/>
 			<Compile Target="ShaderCompileWorker" Configuration="Development" Platform="Linux" Tag="#Build Tools LinuxNative"/>
 			<Compile Target="EpicWebHelper" Configuration="Shipping" Platform="Linux" Tag="#Build Tools LinuxNative"/>
 			<Compile Target="UnrealInsights" Configuration="Development" Platform="Linux" Tag="#Build Tools LinuxNative"/>
