Adding UBA Binaries to the installed engine build

#rnx

diff --git a/Engine/Build/InstalledEngineFilters.xml b/Engine/Build/InstalledEngineFilters.xml
--- a/Engine/Build/InstalledEngineFilters.xml
+++ b/Engine/Build/InstalledEngineFilters.xml
@@ -318,6 +318,9 @@
 		Engine/Binaries/Win64/zen.exe
 		Engine/Binaries/Win64/crashpad_handler.exe
 
+		<!-- Precompiled Win64 UBA -->
+		Engine/Binaries/Win64/UnrealBuildAccelerator/x64/...
+
 		<!-- Used by Switchboard application, and by iOS/tvOS build automation -->
 		Engine/Extras/ThirdPartyNotUE/cwrsync/...
 		Engine/Binaries/ThirdParty/PresentMon/Win64/PresentMon-1.8.0-x64.exe
