Fix for Build DDC jobs to ensure that they include zenserver executables when copying a build to a subdirectory in case it is needed.

diff --git a/Engine/Build/InstalledEngineBuild.xml b/Engine/Build/InstalledEngineBuild.xml
--- a/Engine/Build/InstalledEngineBuild.xml
+++ b/Engine/Build/InstalledEngineBuild.xml
@@ -809,6 +809,11 @@
 			<Tag Files="Engine\Build\BatchFiles\..."  With="#ToCopy"/>
 			<Tag Files="Engine\Binaries\DotNET\Ionic.Zip.Reduced.dll" With="#ToCopy"/>
 			<Tag Files="Engine\Binaries\DotNET\OneSky.dll" With="#ToCopy"/>
+			<Tag Files="Engine\Binaries\Win64\crashpad_handler.exe" With="#ToCopy"/>
+			<Tag Files="Engine\Binaries\Win64\zen.exe" With="#ToCopy"/>
+			<Tag Files="Engine\Binaries\Win64\zen.pdb" With="#ToCopy"/>
+			<Tag Files="Engine\Binaries\Win64\zenserver.exe" With="#ToCopy"/>
+			<Tag Files="Engine\Binaries\Win64\zenserver.pdb" With="#ToCopy"/>
 			<Tag Files="Templates\TemplateResources\..." With="#ToCopy"/>
 			<Expand Name="TagDDCEngineFiles" ExceptPlugins="Plugins\...\Binaries\Mac\..."/>
 			<!-- Filter out the files not needed to build DDC. Removing confidential folders can affect DDC keys, so we want to be sure that we're making DDC with a build that can use it. -->
@@ -857,6 +862,9 @@
 			<Tag Files="Engine/Build/BatchFiles/Mac/..."  With="#ToCopy"/>
 			<Tag Files="Engine/Binaries/DotNET/Ionic.Zip.Reduced.dll" With="#ToCopy"/>
 			<Tag Files="Engine/Binaries/DotNET/OneSky.dll" With="#ToCopy"/>
+			<Tag Files="Engine/Binaries/Mac/crashpad_handler" With="#ToCopy"/>
+			<Tag Files="Engine/Binaries/Mac/zen" With="#ToCopy"/>
+			<Tag Files="Engine/Binaries/Mac/zenserver" With="#ToCopy"/>
 			<Tag Files="Templates/TemplateResources/..." With="#ToCopy"/>
 			<Expand Name="TagDDCEngineFiles" ExceptPlugins="Plugins/.../Binaries/Win64/..."/>
 			<!-- Filter out the files not needed to build DDC. Removing confidential folders can affect DDC keys, so we want to be sure that we're making DDC with a build that can use it. -->
@@ -900,6 +908,8 @@
 			<Tag Files="Engine/Shaders/..." With="#ToCopy"/>
 			<Tag Files="Engine/Binaries/DotNET/Ionic.Zip.Reduced.dll" With="#ToCopy"/>
 			<Tag Files="Engine/Binaries/DotNET/OneSky.dll" With="#ToCopy"/>
+			<Tag Files="Engine/Binaries/Linux/zen" With="#ToCopy"/>
+			<Tag Files="Engine/Binaries/Linux/zenserver" With="#ToCopy"/>
 			<Tag Files="Templates/TemplateResources/..." With="#ToCopy"/>
 			<Expand Name="TagDDCEngineFiles" ExceptPlugins="Plugins/.../Binaries/Mac/..."/>
 			<!-- Filter out the files not needed to build DDC. Removing confidential folders can affect DDC keys, so we want to be sure that we're making DDC with a build that can use it. -->
