* Improved dependency management for editor perf tests. SetupPerfTests now accepts a UseBinariesFrom argument to allow projects to reuse the binaries generated from other projects' compile steps.
* Added a new option for specifying a number of PIE iterations during each editor run. -NumIterationPIE must be a value >0.

diff --git a/Engine/Build/EditorPerf.xml b/Engine/Build/EditorPerf.xml
--- a/Engine/Build/EditorPerf.xml
+++ b/Engine/Build/EditorPerf.xml
@@ -83,8 +83,11 @@
 	<!-- Used to determine if we will include internal projects -->
 	<Option Name="IsLicensee" Restrict="true|false" DefaultValue="false" Description="Define to true if this script is used externally."/>
 
+	<!-- Internal property to track all projects that have been declared -->
+	<Property Name="AllProjects" Value=""/>
+
 	<!-- A macro that can be used to easily set up a new test project. -->
-	<Macro Name="SetupPerfProject" Arguments="ProjectName" OptionalArguments="CopyDefaultsFrom;ProjectTarget;EditorTarget;MapPath;GameSpecificArgs">
+	<Macro Name="SetupPerfProject" Arguments="ProjectName" OptionalArguments="CopyDefaultsFrom;ProjectTarget;EditorTarget;MapPath;GameSpecificArgs;UseBinariesFrom;ExtraNonCopyableArgs">
 		<Property Name="$(ProjectName)__ProjectTarget" Value="$(ProjectName)" CreateInParentScope="true"/>
 		<Property Name="$(ProjectName)__ProjectTarget" Value="$($(CopyDefaultsFrom)__ProjectTarget)" If="'$(CopyDefaultsFrom)' != ''"/>
 		<Property Name="$(ProjectName)__ProjectTarget" Value="$(ProjectTarget)" If="'$(ProjectTarget)' != ''"/>
@@ -100,6 +103,23 @@
 		<Property Name="$(ProjectName)__SpecificArgs" Value="" CreateInParentScope="true"/>
 		<Property Name="$(ProjectName)__SpecificArgs" Value="$($(CopyDefaultsFrom)__SpecificArgs)" If="'$(CopyDefaultsFrom)' != ''"/>
 		<Property Name="$(ProjectName)__SpecificArgs" Value="$(GameSpecificArgs)" If="'$(GameSpecificArgs)' != ''"/>
+
+		<Property Name="$(ProjectName)__UseBinariesFrom" Value="" CreateInParentScope="true"/>
+		<Property Name="$(ProjectName)__UseBinariesFrom" Value="$(CopyDefaultsFrom)" If="'$(CopyDefaultsFrom)' != ''"/>
+		<Property Name="$(ProjectName)__UseBinariesFrom" Value="$(UseBinariesFrom)" If="'$(UseBinariesFrom)' != ''"/>
+
+		<!-- Immediately resolve any indirections on UseBinariesFrom. We cannot support them later when we set up the actual nodes. -->
+		<Do If="'$($(ProjectName)__UseBinariesFrom)' != ''">
+			<Property Name="$(ProjectName)__UseBinariesFrom" Value="$($($(ProjectName)__UseBinariesFrom)__UseBinariesFrom)" If="'$($($(ProjectName)__UseBinariesFrom)__UseBinariesFrom)' != ''"/>
+		</Do>
+
+		<Property Name="$(ProjectName)__ExtraNonCopyableArgs" Value="" CreateInParentScope="true"/>
+		<Property Name="$(ProjectName)__ExtraNonCopyableArgs" Value="$(ExtraNonCopyableArgs)" If="'$(ExtraNonCopyableArgs)' != ''"/>
+
+		<!-- Internal use variable for use with UseBinariesFrom when setting up dependencies-->
+		<Property Name="$(ProjectName)__WillTest" Value="false" CreateInParentScope="true"/>
+
+		<Property Name="AllProjects" Value="$(AllProjects);$(ProjectName)" />
 	</Macro>
 
 	<!-- Inject internal projects default argument at the global scope -->
@@ -116,6 +136,15 @@
 	<Option Name="EnableZenAsyncThreadTest" DefaultValue="true" Description="Enable additional test per phase that uses the Zen Loader."/>
 	<Option Name="EnableZenTestsWithoutEditorDomain" DefaultValue="true" Description="Enable additional test per phase that uses the Zen Loader."/>
 
+	<!-- Option for number of PIE iterations to run (per editor launch) -->
+	<Property Name="DefaultNumIterationsPie" Value="2" />
+	<Option Name="NumIterationsPIE" DefaultValue="$(DefaultNumIterationsPIE)" Description="Specify the number of times to start/stop PIE sessions in a single editor run."/>
+	<Do If="'$(NumIterationsPIE)' &lt;= '0'">
+		<Warning Message="NumIterationPIE must be greater than 0. Setting to default of $(DefaultNumIterationsPIE)."/>
+		<Property Name="NumIterationsPIE" Value="$(DefaultNumIterationsPIE)" />
+	</Do>
+	<Property Name="ArgIterationsPIE" Value="-ini:Engine:[ConsoleVariables]:project.iteration.pie.testiterations=$(NumIterationsPIE)" />
+	
 	<!-- Zen Loader tests specifc arguments -->
 	<Property Name="ZenLoaderTestArgs" Value="-zenloader"/>
 	<Property Name="ZenLoaderAsyncTestArgs" Value="$(ZenLoaderTestArgs) -asyncloadingthread" />
@@ -189,7 +218,7 @@
 		<Property Name="ArtifactName" Value="$(Phase)$(IterationSuffix)$(TestNameSuffix)"/>
 		<Property Name="AllGauntletArgs" Value="-Project=$(ProjectTarget) -Platform=$(Platform) $(GauntletArgs) $(ReportParams) $(LogOutputParams) -artifactname=&quot;$(ArtifactName)&quot;"/>
 		<Property Name="AssetRegistryCacheRootFolderArg" Value="-AssetRegistryCacheRootFolder=&quot;$(AssetRegistryCache_Folder)&quot;"/>
-		<Property Name="CombinedEditorArgs" Value="$(EditorArgs) $(AssetRegistryCacheRootFolderArg)"/>
+		<Property Name="CombinedEditorArgs" Value="$(EditorArgs) $(AssetRegistryCacheRootFolderArg) $(ArgIterationsPIE)"/>
 		<Property Name="RunArgs" Value="$(Args) -editorargs=&quot;$(CombinedEditorArgs)&quot;"/>
 
 		<!-- Do DDC cleanup outside the editor process, handles Cold cleanup without changing stats -->
@@ -267,8 +296,13 @@
 	<!-- ForEach Platform to run, generate graph nodes -->
 	<ForEach Name="Platform" Values="$(Platforms)">
 
-		<!-- ForEach Project to run, generate the graph nodes -->
+		<!-- For each project that was actually requested, set WillTest so that we generate the right test nodes -->
 		<ForEach Name="Project" Values="$(Projects)">
+			<Property Name="$(Project)__WillTest" Value="true"/>
+		</ForEach>
+		
+		<!-- ForEach Project that is defined, generate the graph nodes -->
+		<ForEach Name="Project" Values="$(AllProjects)">
 
 			<!-- Each project should be setup by invoking the SetupPerfProject macro which will set up the various variables we'll need for the tests -->
 			<Property Name="ProjectCombinedArgs" Value="$($(Project)__SpecificArgs) $($(Project)__MapArgs) $(AdditionalRuntimeArgs)"/>
@@ -278,8 +312,8 @@
 			<!-- Use as base string for directories and key for node names -->
 			<Property Name="OutputName" Value="$(Project)$(Platform)EditorPerf"/>
 
-			<!-- Compile the Editor -->
-			<Property Name="CompileNodeName" Value="Compile $(Project) $(Platform)"/>
+			<!-- Compile the Editor. Set up compilation nodes for all projects that have been defined even if we aren't going to test them. This lets the tested projects depend on non-tested ones. -->
+			<Property Name="CompileNodeName" Value="Compile $(Project) $(Platform)" />
 			<Property Name="AgentModifier" Value=""/>
 			<Property Name="AgentModifier" Value=" $(Project)" If="!$(SingleAgentCompile)"/>
 			<Agent Name="Compile Editor Agent $(Platform)$(AgentModifier" Type="$($(Platform)CompileAgentType)">
@@ -291,88 +325,94 @@
 				</Node>
 			</Agent>
 			<Label Category="$(CompileCategory)" Name="$(Project) Editor Compile" Requires="$(CompileNodeName)"/>
-
-			<!-- Run the Tests -->
-			<Property Name="RunNodeRequires" Value="$(CompileNodeName)"/>
-			<Property Name="RunNodeList" Value=""/> <!-- Each node in this the nested loops below, they need to depend on predecessors -->
-			<Property Name="UploadNodeList" Value=""/>
-			<Agent Name="EditorPerf Agent $(Platform) $(Project)" Type="EditorPerf$(Platform)">
-				<ForEach Name="Phase" Values="$(Phases)">
-					<Do If="$(EnableBasePhaseTests)">
-						<ForEach Name="Iteration" Values="$(Iterations)">
-							<!-- Only ever do one iteration of incremental, because it will warm up the workspace -->
-							<Do If="'$(Phase)' != 'Incr' or '$(Iteration)' == '1'">
-								<Property Name="IterationSuffix" Value=""/>
-								<Property Name="IterationSuffix" Value="$(Iteration)" If="'$(Iteration)' != '1'"/>
-								<Property Name="RunNodeName" Value="Run $(OutputName) $(Phase)$(IterationSuffix)"/>
-								<Node Name="$(RunNodeName)" Requires="$(RunNodeRequires);$(RunNodeList)">
-									<Expand Name="RunTest" TestNameSuffix="" Platform="$(Platform)" ProjectName="$(Project)" ProjectTarget="$(ProjectTarget)" NodeName="$(RunNodeName)" Phase="$(Phase)" Iteration="$(Iteration)" Args="$($(Phase)Args)" EditorArgs="$(ProjectCombinedArgs)"/>
-								</Node>
-								<Property Name="RunNodeList" Value="$(RunNodeList);$(RunNodeName)"/>
-							</Do>
-						</ForEach>
-					</Do>	
-					<!-- Don't run Zen tests on incremental runs -->
-					<Do If="'$(Phase)' != 'Incr'">
-						<!-- Zen Loader TestRun if enabled -->
-						<Property Name="RunNodeZenName" Value="Run $(OutputName) $(Phase) Zen"/>
-						<Node Name="$(RunNodeZenName)" Requires="$(RunNodeRequires);$(RunNodeList)" If="$(EnableZenLoaderTest)">
-							<Expand Name="RunTest" TestNameSuffix="Zen" Platform="$(Platform)" ProjectName="$(Project)" ProjectTarget="$(ProjectTarget)" NodeName="$(RunNodeZenName)" Phase="$(Phase)" Iteration="1" Args="$($(Phase)Args)" EditorArgs="$(ProjectCombinedArgs) $(ZenLoaderTestArgs)"/>
-						</Node>
-						<Property Name="RunNodeList" Value="$(RunNodeList);$(RunNodeZenName)" If="$(EnableZenLoaderTest)"/>
-						<!-- Zen Loader Async TestRun if enabled -->
-						<Property Name="RunNodeZenAsyncName" Value="Run $(OutputName) $(Phase) Zen Async"/>
-						<Node Name="$(RunNodeZenAsyncName)" Requires="$(RunNodeRequires);$(RunNodeList)" If="$(EnableZenAsyncThreadTest)">
-							<Expand Name="RunTest" TestNameSuffix="ZenAsync" Platform="$(Platform)" ProjectName="$(Project)" ProjectTarget="$(ProjectTarget)" NodeName="$(RunNodeZenAsyncName)" Phase="$(Phase)" Iteration="1" Args="$($(Phase)Args)" EditorArgs="$(ProjectCombinedArgs) $(ZenLoaderAsyncTestArgs)"/>
-						</Node>
-						<Property Name="RunNodeList" Value="$(RunNodeList);$(RunNodeZenAsyncName)" If="$(EnableZenAsyncThreadTest)"/>
-					
-						<!-- Zen Store DDC TestRun if enabled -->
-						<Property Name="RunNodeZenStoreDDCName" Value="Run $(OutputName) $(Phase) Zen Store DDC"/>
-						<Node Name="$(RunNodeZenStoreDDCName)" Requires="$(RunNodeRequires);$(RunNodeList)" If="$(EnableZenStoreDDCTest)">
-							<Expand Name="RunTest" TestNameSuffix="ZenAsync" Platform="$(Platform)" ProjectName="$(Project)" ProjectTarget="$(ProjectTarget)" NodeName="$(RunNodeZenStoreDDCName)" Phase="$(Phase)" Iteration="1" Args="$(ZenStoreDDCPhaseArgs)" EditorArgs="$(ProjectCombinedArgs) $(ZenStoreDDCTestArgs)"/>
-						</Node>
-						<Property Name="RunNodeList" Value="$(RunNodeList);$(RunNodeZenStoreDDCName)" If="$(EnableZenStoreDDCTest)"/>
-					</Do>
-					<!-- Warm only test runs e.g LazyLoad -->
-					<Do If="'$(Phase)' == 'Warm'">
-						<!-- Legacy LazyLoad TestRun if enabled -->
-						<Property Name="RunNodeLazyLoadName" Value="Run $(OutputName) $(Phase) LazyLoad"/>
-						<Node Name="$(RunNodeLazyLoadName)" Requires="$(RunNodeRequires);$(RunNodeList)" If="$(EnableLazyLoadTest)">
-							<Expand Name="RunTest" TestNameSuffix="Zen" Platform="$(Platform)" ProjectName="$(Project)" ProjectTarget="$(ProjectTarget)" NodeName="$(RunNodeLazyLoadName)" Phase="$(Phase)" Iteration="1" Args="$($(Phase)Args)" EditorArgs="$(ProjectCombinedArgs) $(LazyLoadTestArgs)"/>
-						</Node>
-						<Property Name="RunNodeList" Value="$(RunNodeList);$(RunNodeLazyLoadName)" If="$(EnableLazyLoadTest)"/>
-						<!-- Zen LazyLoad TestRun if enabled -->
-						<Property Name="RunNodeLazyLoadZenName" Value="Run $(OutputName) $(Phase) LazyLoad Zen"/>
-						<Node Name="$(RunNodeLazyLoadZenName)" Requires="$(RunNodeRequires);$(RunNodeList)" If="$(EnableLazyLoadZenTest)">
-							<Expand Name="RunTest" TestNameSuffix="Zen" Platform="$(Platform)" ProjectName="$(Project)" ProjectTarget="$(ProjectTarget)" NodeName="$(RunNodeLazyLoadZenName)" Phase="$(Phase)" Iteration="1" Args="$($(Phase)Args)" EditorArgs="$(ProjectCombinedArgs) $(LazyLoadZenTestArgs)"/>
-						</Node>
-						<Property Name="RunNodeList" Value="$(RunNodeList);$(RunNodeLazyLoadZenName)" If="$(EnableLazyLoadZenTest)"/>
-					</Do>
-				</ForEach>
-				<Property Name="RunNodeUploadPhaseTraces" Value="Upload $(OutputName) Traces"/>
-				<Node Name="$(RunNodeUploadPhaseTraces)">
-					<Expand Name="UploadTestTraces" Platform="$(Platform)" ProjectName="$(Project)"/>
-				</Node>
-				<Property Name="UploadNodeList" Value="$(RunNodeUploadPhaseTraces)"/>
-			</Agent>
-			<Label Category="$(RunCategory)" Name="$(Project) Editor Run" Requires="$(RunNodeRequires);$(RunNodeList);$(UploadNodeList)" UgsBadge="$(Project) Editor Run"/>
-			<Property Name="RunGraphNodes" Value="$(RunGraphNodes);$(RunNodeList);$(UploadNodeList)"/>
-
-			<!-- Baseline the Results -->
-			<Property Name="BaselineNodeName" Value="Baseline $(OutputName)"/>
-			<Property Name="BaselineNodeRequires" Value=""/>
-			<Property Name="BaselineNodeRequires" Value="$(RunNodeName)" If="$(RunAndBaseline)"/>
-			<Property Name="BaselineGraphNodes" Value="$(BaselineGraphNodes);$(BaselineNodeName)"/>
-			<Agent Name="EditorPerf Agent $(Platform) $(Project)" Type="EditorPerf$(Platform)">
-				<Node Name="$(BaselineNodeName)" Requires="$(BaselineNodeRequires)">
+			<!-- Store off the name of the compile node so that we can use look it up by Project name from UseBinariesFrom. -->
+			<Property Name="$(Project)__CompileNodeName" Value="$(CompileNodeName)"/>
+
+			<!-- If the project will be run, set up the test nodes -->
+			<Do If="'$($(Project)__WillTest)'=='true'">
+				<!-- Run the Tests -->
+				<Property Name="RunNodeRequires" Value="$(CompileNodeName)"/>
+				<Property Name="RunNodeRequires" Value="$($($(Project)__UseBinariesFrom)__CompileNodeName)" If="'$($(Project)__UseBinariesFrom)'!=''"/>
+				<Property Name="RunNodeList" Value=""/>
+				<!--Each node in this the nested loops below, they need to depend on predecessors--> 
+				<Property Name="UploadNodeList" Value=""/>
+				<Agent Name="EditorPerf Agent $(Platform) $(Project)" Type="EditorPerf$(Platform)">
 					<ForEach Name="Phase" Values="$(Phases)">
-						<Expand Name="BaselineResult" Platform="$(Platform)" Project="$(Project)" Phase="$(Phase)"/>
+						<Do If="$(EnableBasePhaseTests)">
+							<ForEach Name="Iteration" Values="$(Iterations)">
+								<!--Only ever do one iteration of incremental, because it will warm up the workspace--> 
+								<Do If="'$(Phase)' != 'Incr' or '$(Iteration)' == '1'">
+									<Property Name="IterationSuffix" Value=""/>
+									<Property Name="IterationSuffix" Value="$(Iteration)" If="'$(Iteration)' != '1'"/>
+									<Property Name="RunNodeName" Value="Run $(OutputName) $(Phase)$(IterationSuffix)"/>
+									<Node Name="$(RunNodeName)" Requires="$(RunNodeRequires);$(RunNodeList)">
+										<Expand Name="RunTest" TestNameSuffix="" Platform="$(Platform)" ProjectName="$(Project)" ProjectTarget="$(ProjectTarget)" NodeName="$(RunNodeName)" Phase="$(Phase)" Iteration="$(Iteration)" Args="$($(Phase)Args)" EditorArgs="$(ProjectCombinedArgs)"/>
+									</Node>
+									<Property Name="RunNodeList" Value="$(RunNodeList);$(RunNodeName)"/>
+								</Do>
+							</ForEach>
+						</Do>
+						<!--Don't run Zen tests on incremental runs--> 
+						<Do If="'$(Phase)' != 'Incr'">
+							<!--Zen Loader TestRun if enabled--> 
+							<Property Name="RunNodeZenName" Value="Run $(OutputName) $(Phase) Zen"/>
+							<Node Name="$(RunNodeZenName)" Requires="$(RunNodeRequires);$(RunNodeList)" If="$(EnableZenLoaderTest)">
+								<Expand Name="RunTest" TestNameSuffix="Zen" Platform="$(Platform)" ProjectName="$(Project)" ProjectTarget="$(ProjectTarget)" NodeName="$(RunNodeZenName)" Phase="$(Phase)" Iteration="1" Args="$($(Phase)Args)" EditorArgs="$(ProjectCombinedArgs) $(ZenLoaderTestArgs)"/>
+							</Node>
+							<Property Name="RunNodeList" Value="$(RunNodeList);$(RunNodeZenName)" If="$(EnableZenLoaderTest)"/>
+							<!--Zen Loader Async TestRun if enabled--> 
+							<Property Name="RunNodeZenAsyncName" Value="Run $(OutputName) $(Phase) Zen Async"/>
+							<Node Name="$(RunNodeZenAsyncName)" Requires="$(RunNodeRequires);$(RunNodeList)" If="$(EnableZenAsyncThreadTest)">
+								<Expand Name="RunTest" TestNameSuffix="ZenAsync" Platform="$(Platform)" ProjectName="$(Project)" ProjectTarget="$(ProjectTarget)" NodeName="$(RunNodeZenAsyncName)" Phase="$(Phase)" Iteration="1" Args="$($(Phase)Args)" EditorArgs="$(ProjectCombinedArgs) $(ZenLoaderAsyncTestArgs)"/>
+							</Node>
+							<Property Name="RunNodeList" Value="$(RunNodeList);$(RunNodeZenAsyncName)" If="$(EnableZenAsyncThreadTest)"/>
+
+							<!--Zen Store DDC TestRun if enabled--> 
+							<Property Name="RunNodeZenStoreDDCName" Value="Run $(OutputName) $(Phase) Zen Store DDC"/>
+							<Node Name="$(RunNodeZenStoreDDCName)" Requires="$(RunNodeRequires);$(RunNodeList)" If="$(EnableZenStoreDDCTest)">
+								<Expand Name="RunTest" TestNameSuffix="ZenAsync" Platform="$(Platform)" ProjectName="$(Project)" ProjectTarget="$(ProjectTarget)" NodeName="$(RunNodeZenStoreDDCName)" Phase="$(Phase)" Iteration="1" Args="$(ZenStoreDDCPhaseArgs)" EditorArgs="$(ProjectCombinedArgs) $(ZenStoreDDCTestArgs)"/>
+							</Node>
+							<Property Name="RunNodeList" Value="$(RunNodeList);$(RunNodeZenStoreDDCName)" If="$(EnableZenStoreDDCTest)"/>
+						</Do>
+						<!--Warm only test runs e.g LazyLoad--> 
+						<Do If="'$(Phase)' == 'Warm'">
+							<!--Legacy LazyLoad TestRun if enabled--> 
+							<Property Name="RunNodeLazyLoadName" Value="Run $(OutputName) $(Phase) LazyLoad"/>
+							<Node Name="$(RunNodeLazyLoadName)" Requires="$(RunNodeRequires);$(RunNodeList)" If="$(EnableLazyLoadTest)">
+								<Expand Name="RunTest" TestNameSuffix="Zen" Platform="$(Platform)" ProjectName="$(Project)" ProjectTarget="$(ProjectTarget)" NodeName="$(RunNodeLazyLoadName)" Phase="$(Phase)" Iteration="1" Args="$($(Phase)Args)" EditorArgs="$(ProjectCombinedArgs) $(LazyLoadTestArgs)"/>
+							</Node>
+							<Property Name="RunNodeList" Value="$(RunNodeList);$(RunNodeLazyLoadName)" If="$(EnableLazyLoadTest)"/>
+							<!--Zen LazyLoad TestRun if enabled--> 
+							<Property Name="RunNodeLazyLoadZenName" Value="Run $(OutputName) $(Phase) LazyLoad Zen"/>
+							<Node Name="$(RunNodeLazyLoadZenName)" Requires="$(RunNodeRequires);$(RunNodeList)" If="$(EnableLazyLoadZenTest)">
+								<Expand Name="RunTest" TestNameSuffix="Zen" Platform="$(Platform)" ProjectName="$(Project)" ProjectTarget="$(ProjectTarget)" NodeName="$(RunNodeLazyLoadZenName)" Phase="$(Phase)" Iteration="1" Args="$($(Phase)Args)" EditorArgs="$(ProjectCombinedArgs) $(LazyLoadZenTestArgs)"/>
+							</Node>
+							<Property Name="RunNodeList" Value="$(RunNodeList);$(RunNodeLazyLoadZenName)" If="$(EnableLazyLoadZenTest)"/>
+						</Do>
 					</ForEach>
-				</Node>
-			</Agent>
-			<Label Category="$(BaselineCategory)" Name="$(Project) Perf Baseline" Requires="$(BaselineNodeName)"/>
-
+					<Property Name="RunNodeUploadPhaseTraces" Value="Upload $(OutputName) Traces"/>
+					<Node Name="$(RunNodeUploadPhaseTraces)">
+						<Expand Name="UploadTestTraces" Platform="$(Platform)" ProjectName="$(Project)"/>
+					</Node>
+					<Property Name="UploadNodeList" Value="$(RunNodeUploadPhaseTraces)"/>
+				</Agent>
+				<Label Category="$(RunCategory)" Name="$(Project) Editor Run" Requires="$(RunNodeRequires);$(RunNodeList);$(UploadNodeList)" UgsBadge="$(Project) Editor Run"/>
+				<Property Name="RunGraphNodes" Value="$(RunGraphNodes);$(RunNodeList);$(UploadNodeList)"/>
+
+				<!--Baseline the Results--> 
+				<Property Name="BaselineNodeName" Value="Baseline $(OutputName)"/>
+				<Property Name="BaselineNodeRequires" Value=""/>
+				<Property Name="BaselineNodeRequires" Value="$(RunNodeName)" If="$(RunAndBaseline)"/>
+				<Property Name="BaselineGraphNodes" Value="$(BaselineGraphNodes);$(BaselineNodeName)"/>
+				<Agent Name="EditorPerf Agent $(Platform) $(Project)" Type="EditorPerf$(Platform)">
+					<Node Name="$(BaselineNodeName)" Requires="$(BaselineNodeRequires)">
+						<ForEach Name="Phase" Values="$(Phases)">
+							<Expand Name="BaselineResult" Platform="$(Platform)" Project="$(Project)" Phase="$(Phase)"/>
+						</ForEach>
+					</Node>
+				</Agent>
+				<Label Category="$(BaselineCategory)" Name="$(Project) Perf Baseline" Requires="$(BaselineNodeName)"/>
+			</Do>
 		</ForEach> <!-- Projects -->
 
 	</ForEach> <!-- Platforms -->
