Small bug fix to CL 31276995 to actually use the ExtraNonCopyableArgs added to EditorPerf.xml perf projects.

diff --git a/Engine/Build/EditorPerf.xml b/Engine/Build/EditorPerf.xml
--- a/Engine/Build/EditorPerf.xml
+++ b/Engine/Build/EditorPerf.xml
@@ -305,7 +305,7 @@
 		<ForEach Name="Project" Values="$(AllProjects)">
 
 			<!-- Each project should be setup by invoking the SetupPerfProject macro which will set up the various variables we'll need for the tests -->
-			<Property Name="ProjectCombinedArgs" Value="$($(Project)__SpecificArgs) $($(Project)__MapArgs) $(AdditionalRuntimeArgs)"/>
+			<Property Name="ProjectCombinedArgs" Value="$($(Project)__SpecificArgs) $($(Project)__MapArgs) $($(Project)__ExtraNonCopyableArgs) $(AdditionalRuntimeArgs)"/>
 			<Property Name="ProjectTarget" Value="$($(Project)__ProjectTarget)"/>
 			<Property Name="EditorTarget" Value="$($(Project)__EditorTarget)"/>
 
