Added a -nozenloader parameter to EditorPerf due to ZenLoader being on by default.

diff --git a/Engine/Build/EditorPerf.xml b/Engine/Build/EditorPerf.xml
--- a/Engine/Build/EditorPerf.xml
+++ b/Engine/Build/EditorPerf.xml
@@ -147,6 +147,7 @@
 	
 	<!-- Zen Loader tests specifc arguments -->
 	<Property Name="ZenLoaderTestArgs" Value="-zenloader"/>
+	<Property Name="NoZenLoaderTestArgs" Value="-nozenloader"/>
 	<Property Name="ZenLoaderAsyncTestArgs" Value="$(ZenLoaderTestArgs) -asyncloadingthread" />
 	
 	<Property Name="ZenLoaderTestArgs" Value="$(ZenLoaderTestArgs) -noeditordomain" If="$(EnableZenTestsWithoutEditorDomain)" />
@@ -346,7 +347,7 @@
 									<Property Name="IterationSuffix" Value="$(Iteration)" If="'$(Iteration)' != '1'"/>
 									<Property Name="RunNodeName" Value="Run $(OutputName) $(Phase)$(IterationSuffix)"/>
 									<Node Name="$(RunNodeName)" Requires="$(RunNodeRequires);$(RunNodeList)">
-										<Expand Name="RunTest" TestNameSuffix="" Platform="$(Platform)" ProjectName="$(Project)" ProjectTarget="$(ProjectTarget)" NodeName="$(RunNodeName)" Phase="$(Phase)" Iteration="$(Iteration)" Args="$($(Phase)Args)" EditorArgs="$(ProjectCombinedArgs)"/>
+										<Expand Name="RunTest" TestNameSuffix="" Platform="$(Platform)" ProjectName="$(Project)" ProjectTarget="$(ProjectTarget)" NodeName="$(RunNodeName)" Phase="$(Phase)" Iteration="$(Iteration)" Args="$($(Phase)Args)" EditorArgs="$(ProjectCombinedArgs) $(NoZenLoaderTestArgs)"/>
 									</Node>
 									<Property Name="RunNodeList" Value="$(RunNodeList);$(RunNodeName)"/>
 								</Do>
@@ -379,13 +380,13 @@
 							<!--Legacy LazyLoad TestRun if enabled--> 
 							<Property Name="RunNodeLazyLoadName" Value="Run $(OutputName) $(Phase) LazyLoad"/>
 							<Node Name="$(RunNodeLazyLoadName)" Requires="$(RunNodeRequires);$(RunNodeList)" If="$(EnableLazyLoadTest)">
-								<Expand Name="RunTest" TestNameSuffix="Zen" Platform="$(Platform)" ProjectName="$(Project)" ProjectTarget="$(ProjectTarget)" NodeName="$(RunNodeLazyLoadName)" Phase="$(Phase)" Iteration="1" Args="$($(Phase)Args)" EditorArgs="$(ProjectCombinedArgs) $(LazyLoadTestArgs)"/>
+								<Expand Name="RunTest" TestNameSuffix="NoZenLoader" Platform="$(Platform)" ProjectName="$(Project)" ProjectTarget="$(ProjectTarget)" NodeName="$(RunNodeLazyLoadName)" Phase="$(Phase)" Iteration="1" Args="$($(Phase)Args)" EditorArgs="$(ProjectCombinedArgs) $(LazyLoadTestArgs) $(NoZenLoaderTestArgs)"/>
 							</Node>
 							<Property Name="RunNodeList" Value="$(RunNodeList);$(RunNodeLazyLoadName)" If="$(EnableLazyLoadTest)"/>
 							<!--Zen LazyLoad TestRun if enabled--> 
 							<Property Name="RunNodeLazyLoadZenName" Value="Run $(OutputName) $(Phase) LazyLoad Zen"/>
 							<Node Name="$(RunNodeLazyLoadZenName)" Requires="$(RunNodeRequires);$(RunNodeList)" If="$(EnableLazyLoadZenTest)">
-								<Expand Name="RunTest" TestNameSuffix="Zen" Platform="$(Platform)" ProjectName="$(Project)" ProjectTarget="$(ProjectTarget)" NodeName="$(RunNodeLazyLoadZenName)" Phase="$(Phase)" Iteration="1" Args="$($(Phase)Args)" EditorArgs="$(ProjectCombinedArgs) $(LazyLoadZenTestArgs)"/>
+								<Expand Name="RunTest" TestNameSuffix="ZenLoader" Platform="$(Platform)" ProjectName="$(Project)" ProjectTarget="$(ProjectTarget)" NodeName="$(RunNodeLazyLoadZenName)" Phase="$(Phase)" Iteration="1" Args="$($(Phase)Args)" EditorArgs="$(ProjectCombinedArgs) $(LazyLoadZenTestArgs)"/>
 							</Node>
 							<Property Name="RunNodeList" Value="$(RunNodeList);$(RunNodeLazyLoadZenName)" If="$(EnableLazyLoadZenTest)"/>
 						</Do>
