Removing Incremental phase from EditorPerfTest as it never functioned properly. General cleanup of naming and logic around phases and zen loader. Cleared up an issue with AssetRegistryCache where we would have had a warm registry cache on Cold phases after the first. These changes require corresponding changes to the stream json.

diff --git a/Engine/Build/EditorPerf.xml b/Engine/Build/EditorPerf.xml
--- a/Engine/Build/EditorPerf.xml
+++ b/Engine/Build/EditorPerf.xml
@@ -66,15 +66,17 @@
 	<Property Name="BaselineAutomationNetworkOutputDir" Value="$(AutomationNetworkRootDir)/$(BaselineBuildName)"/>
 	<Property Name="BaselineAutomationNetworkOutputDir" Value="$(AutomationOutputDir)" If="$(RunAndBaseline)"/>
 
+	<!-- Options for warm+cold vs cold cache only -->
+	<Option Name="IncludeWarmPhase" DefaultValue="true" Description="Whether to include a warm cache phase. Default is true and will run cold and warm tests. If false, only a cold test will run."/>
+	<Property Name="Phases" Value="Cold;Warm"/>
+	<Property Name="Phases" Value="Cold" If="'$(IncludeWarmPhase)'=='false'"/>
+
 	<!-- Arguments for tracing -->
-	<Option Name="Phases" DefaultValue="Incr;Warm;Cold" Description="The list of phases to perform"/>
-	<Error Message="No phases specified" If="'$(Phases)' == ''"/>
 	<Property Name="TraceDir" Value="$(RootDir)/Engine/Programs/UnrealInsights/Saved/TraceSessions"/>
 
 	<!-- Arguments for phases, keep in mind that these are passed to both the DDCCleanupCommandlet and RunUnreal -->
 	<Option Name="EnableBasePhaseTests" DefaultValue="true" Description="Enable the base phase perf tests."/>
-	<Option Name="ColdArgs" DefaultValue="-ddc=Cold" Description="Arguments to unreal for the cold phase"/>
-	<Option Name="IncrArgs" DefaultValue="-ddc=Warm" Description="Arguments to unreal for the incremental phase"/>
+	<Option Name="ColdArgs" DefaultValue="-ddc=Cold -NoAssetRegistryCacheRead" Description="Arguments to unreal for the cold phase"/>
 	<Option Name="WarmArgs" DefaultValue="-ddc=Warm" Description="Arguments to unreal for the warm phase"/>
 
 	<!-- Iterations to run and average together -->
@@ -116,16 +118,13 @@
 		<Property Name="$(ProjectName)__ExtraNonCopyableArgs" Value="" CreateInParentScope="true"/>
 		<Property Name="$(ProjectName)__ExtraNonCopyableArgs" Value="$(ExtraNonCopyableArgs)" If="'$(ExtraNonCopyableArgs)' != ''"/>
 
-		<!-- Internal use variable for use with UseBinariesFrom when setting up dependencies-->
-		<Property Name="$(ProjectName)__WillTest" Value="false" CreateInParentScope="true"/>
-
 		<Property Name="AllProjects" Value="$(AllProjects);$(ProjectName)" />
 	</Macro>
 
 	<!-- Inject internal projects default argument at the global scope -->
 	<Include Script="../Restricted/NotForLicensees/Build/EditorPerfInternalArgs.xml" If="!$(IsLicensee)" />
 
-	<!-- Option to retain per-branch asset registry between runs so that incremental is representative of user experience -->
+	<!-- Option to retain per-branch asset registry between runs so that warm is representative of user experience -->
 	<Property Name="AssetRegistryCache_Folder" Value="$(RootDir)/../../../EditorPerfAssetRegistryCache/$(EscapedBranch)/"/>
 
 	<!-- Options for extra runtime args -->
@@ -133,8 +132,8 @@
 
 	<!-- Options for extra Zen Loader tests -->
 	<Option Name="EnableZenLoaderTest" DefaultValue="true" Description="Enable additional test per phase that uses the Zen Loader."/>
-	<Option Name="EnableZenAsyncThreadTest" DefaultValue="true" Description="Enable additional test per phase that uses the Zen Loader."/>
-	<Option Name="EnableZenTestsWithoutEditorDomain" DefaultValue="true" Description="Enable additional test per phase that uses the Zen Loader."/>
+	<Option Name="EnableZenLoaderAsyncThreadTest" DefaultValue="false" Description="Enable additional test per phase that uses the Zen Loader with an async loading thread."/>
+	<Option Name="DisableEditorDomainForZenLoaderTests" DefaultValue="true" Description="Explicitly disables editordomain (-noeditordomain) for zen loader tests"/>
 
 	<!-- Option for number of PIE iterations to run (per editor launch) -->
 	<Property Name="DefaultNumIterationsPie" Value="2" />
@@ -150,8 +149,8 @@
 	<Property Name="NoZenLoaderTestArgs" Value="-nozenloader"/>
 	<Property Name="ZenLoaderAsyncTestArgs" Value="$(ZenLoaderTestArgs) -asyncloadingthread" />
 	
-	<Property Name="ZenLoaderTestArgs" Value="$(ZenLoaderTestArgs) -noeditordomain" If="$(EnableZenTestsWithoutEditorDomain)" />
-	<Property Name="ZenLoaderAsyncTestArgs" Value="$(ZenLoaderAsyncTestArgs) -noeditordomain" If="$(EnableZenTestsWithoutEditorDomain)" />
+	<Property Name="ZenLoaderTestArgs" Value="$(ZenLoaderTestArgs) -noeditordomain" If="$(DisableEditorDomainForZenLoaderTests)" />
+	<Property Name="ZenLoaderAsyncTestArgs" Value="$(ZenLoaderAsyncTestArgs) -noeditordomain" If="$(DisableEditorDomainForZenLoaderTests)" />
 
 	<!-- Options for extra Zen Store tests -->
 	<Option Name="EnableZenStoreDDCTest" DefaultValue="false" Description="Enable additional test on warm phase that uses the Zen Store DDC."/>
@@ -163,8 +162,8 @@
 	<Property Name="ZenStoreDDCTestArgs" Value="$(ZenStoreCommonArgs)"/>
 
 	<!-- Options for extra LazyLoad tests -->
-	<Option Name="EnableLazyLoadTest" DefaultValue="false" Description="Enable additional test on warm phase that uses LazyLoad."/>
-	<Option Name="EnableLazyLoadZenTest" DefaultValue="false" Description="Enable additional test on warm phase that uses LazyLoad with Zen Loader."/>
+	<Option Name="EnableLazyLoadNoZenLoaderTest" DefaultValue="false" Description="Enable additional test on warm phase that uses LazyLoad with Zen Loader disabled."/>
+	<Option Name="EnableLazyLoadZenLoaderTest" DefaultValue="false" Description="Enable additional test on warm phase that uses LazyLoad with Zen Loader enabled."/>
 	
 	<!-- LazyLoad tests specific arguments -->
 	<Property Name="LazyLoadTestArgs" Value="-LazyLoadImports"/>
@@ -297,11 +296,6 @@
 	<!-- ForEach Platform to run, generate graph nodes -->
 	<ForEach Name="Platform" Values="$(Platforms)">
 
-		<!-- For each project that was actually requested, set WillTest so that we generate the right test nodes -->
-		<ForEach Name="Project" Values="$(Projects)">
-			<Property Name="$(Project)__WillTest" Value="true"/>
-		</ForEach>
-		
 		<!-- ForEach Project that is defined, generate the graph nodes -->
 		<ForEach Name="Project" Values="$(AllProjects)">
 
@@ -330,7 +324,7 @@
 			<Property Name="$(Project)__CompileNodeName" Value="$(CompileNodeName)"/>
 
 			<!-- If the project will be run, set up the test nodes -->
-			<Do If="'$($(Project)__WillTest)'=='true'">
+			<Do If="ContainsItem('$(Projects)', '$(Project)', ';')=='true'">
 				<!-- Run the Tests -->
 				<Property Name="RunNodeRequires" Value="$(CompileNodeName)"/>
 				<Property Name="RunNodeRequires" Value="$($($(Project)__UseBinariesFrom)__CompileNodeName)" If="'$($(Project)__UseBinariesFrom)'!=''"/>
@@ -341,54 +335,51 @@
 					<ForEach Name="Phase" Values="$(Phases)">
 						<Do If="$(EnableBasePhaseTests)">
 							<ForEach Name="Iteration" Values="$(Iterations)">
-								<!--Only ever do one iteration of incremental, because it will warm up the workspace--> 
-								<Do If="'$(Phase)' != 'Incr' or '$(Iteration)' == '1'">
-									<Property Name="IterationSuffix" Value=""/>
-									<Property Name="IterationSuffix" Value="$(Iteration)" If="'$(Iteration)' != '1'"/>
-									<Property Name="RunNodeName" Value="Run $(OutputName) $(Phase)$(IterationSuffix)"/>
-									<Node Name="$(RunNodeName)" Requires="$(RunNodeRequires);$(RunNodeList)">
-										<Expand Name="RunTest" TestNameSuffix="" Platform="$(Platform)" ProjectName="$(Project)" ProjectTarget="$(ProjectTarget)" NodeName="$(RunNodeName)" Phase="$(Phase)" Iteration="$(Iteration)" Args="$($(Phase)Args)" EditorArgs="$(ProjectCombinedArgs) $(NoZenLoaderTestArgs)"/>
-									</Node>
-									<Property Name="RunNodeList" Value="$(RunNodeList);$(RunNodeName)"/>
-								</Do>
+								<Property Name="IterationSuffix" Value=""/>
+								<Property Name="IterationSuffix" Value="$(Iteration)" If="'$(Iteration)' != '1'"/>
+								<Property Name="RunNodeName" Value="Run $(OutputName) $(Phase)$(IterationSuffix)"/>
+								<Node Name="$(RunNodeName)" Requires="$(RunNodeRequires);$(RunNodeList)">
+									<Expand Name="RunTest" TestNameSuffix="NoZenLoader" Platform="$(Platform)" ProjectName="$(Project)" ProjectTarget="$(ProjectTarget)" NodeName="$(RunNodeName)" Phase="$(Phase)" Iteration="$(Iteration)" Args="$($(Phase)Args)" EditorArgs="$(ProjectCombinedArgs) $(NoZenLoaderTestArgs)"/>
+								</Node>
+								<Property Name="RunNodeList" Value="$(RunNodeList);$(RunNodeName)"/>
 							</ForEach>
 						</Do>
-						<!--Don't run Zen tests on incremental runs--> 
-						<Do If="'$(Phase)' != 'Incr'">
+						<!-- Setup Zen Loader tests for this phase -->
+						<Do>
 							<!--Zen Loader TestRun if enabled--> 
-							<Property Name="RunNodeZenName" Value="Run $(OutputName) $(Phase) Zen"/>
+							<Property Name="RunNodeZenName" Value="Run $(OutputName) $(Phase) ZenLoader"/>
 							<Node Name="$(RunNodeZenName)" Requires="$(RunNodeRequires);$(RunNodeList)" If="$(EnableZenLoaderTest)">
-								<Expand Name="RunTest" TestNameSuffix="Zen" Platform="$(Platform)" ProjectName="$(Project)" ProjectTarget="$(ProjectTarget)" NodeName="$(RunNodeZenName)" Phase="$(Phase)" Iteration="1" Args="$($(Phase)Args)" EditorArgs="$(ProjectCombinedArgs) $(ZenLoaderTestArgs)"/>
+								<Expand Name="RunTest" TestNameSuffix="ZenLoader" Platform="$(Platform)" ProjectName="$(Project)" ProjectTarget="$(ProjectTarget)" NodeName="$(RunNodeZenName)" Phase="$(Phase)" Iteration="1" Args="$($(Phase)Args)" EditorArgs="$(ProjectCombinedArgs) $(ZenLoaderTestArgs)"/>
 							</Node>
 							<Property Name="RunNodeList" Value="$(RunNodeList);$(RunNodeZenName)" If="$(EnableZenLoaderTest)"/>
 							<!--Zen Loader Async TestRun if enabled--> 
-							<Property Name="RunNodeZenAsyncName" Value="Run $(OutputName) $(Phase) Zen Async"/>
-							<Node Name="$(RunNodeZenAsyncName)" Requires="$(RunNodeRequires);$(RunNodeList)" If="$(EnableZenAsyncThreadTest)">
-								<Expand Name="RunTest" TestNameSuffix="ZenAsync" Platform="$(Platform)" ProjectName="$(Project)" ProjectTarget="$(ProjectTarget)" NodeName="$(RunNodeZenAsyncName)" Phase="$(Phase)" Iteration="1" Args="$($(Phase)Args)" EditorArgs="$(ProjectCombinedArgs) $(ZenLoaderAsyncTestArgs)"/>
+							<Property Name="RunNodeZenAsyncName" Value="Run $(OutputName) $(Phase) ZenLoader Async"/>
+							<Node Name="$(RunNodeZenAsyncName)" Requires="$(RunNodeRequires);$(RunNodeList)" If="$(EnableZenLoaderAsyncThreadTest)">
+								<Expand Name="RunTest" TestNameSuffix="ZenLoaderAsync" Platform="$(Platform)" ProjectName="$(Project)" ProjectTarget="$(ProjectTarget)" NodeName="$(RunNodeZenAsyncName)" Phase="$(Phase)" Iteration="1" Args="$($(Phase)Args)" EditorArgs="$(ProjectCombinedArgs) $(ZenLoaderAsyncTestArgs)"/>
 							</Node>
-							<Property Name="RunNodeList" Value="$(RunNodeList);$(RunNodeZenAsyncName)" If="$(EnableZenAsyncThreadTest)"/>
+							<Property Name="RunNodeList" Value="$(RunNodeList);$(RunNodeZenAsyncName)" If="$(EnableZenLoaderAsyncThreadTest)"/>
 
 							<!--Zen Store DDC TestRun if enabled--> 
-							<Property Name="RunNodeZenStoreDDCName" Value="Run $(OutputName) $(Phase) Zen Store DDC"/>
+							<Property Name="RunNodeZenStoreDDCName" Value="Run $(OutputName) $(Phase) ZenLoader Store DDC"/>
 							<Node Name="$(RunNodeZenStoreDDCName)" Requires="$(RunNodeRequires);$(RunNodeList)" If="$(EnableZenStoreDDCTest)">
-								<Expand Name="RunTest" TestNameSuffix="ZenAsync" Platform="$(Platform)" ProjectName="$(Project)" ProjectTarget="$(ProjectTarget)" NodeName="$(RunNodeZenStoreDDCName)" Phase="$(Phase)" Iteration="1" Args="$(ZenStoreDDCPhaseArgs)" EditorArgs="$(ProjectCombinedArgs) $(ZenStoreDDCTestArgs)"/>
+								<Expand Name="RunTest" TestNameSuffix="ZenStoreAsyncDDC" Platform="$(Platform)" ProjectName="$(Project)" ProjectTarget="$(ProjectTarget)" NodeName="$(RunNodeZenStoreDDCName)" Phase="$(Phase)" Iteration="1" Args="$(ZenStoreDDCPhaseArgs)" EditorArgs="$(ProjectCombinedArgs) $(ZenStoreDDCTestArgs)"/>
 							</Node>
 							<Property Name="RunNodeList" Value="$(RunNodeList);$(RunNodeZenStoreDDCName)" If="$(EnableZenStoreDDCTest)"/>
 						</Do>
 						<!--Warm only test runs e.g LazyLoad--> 
 						<Do If="'$(Phase)' == 'Warm'">
 							<!--Legacy LazyLoad TestRun if enabled--> 
-							<Property Name="RunNodeLazyLoadName" Value="Run $(OutputName) $(Phase) LazyLoad"/>
-							<Node Name="$(RunNodeLazyLoadName)" Requires="$(RunNodeRequires);$(RunNodeList)" If="$(EnableLazyLoadTest)">
+							<Property Name="RunNodeLazyLoadName" Value="Run $(OutputName) $(Phase) LazyLoad NoZenLoader"/>
+							<Node Name="$(RunNodeLazyLoadName)" Requires="$(RunNodeRequires);$(RunNodeList)" If="$(EnableLazyLoadNoZenLoaderTest)">
 								<Expand Name="RunTest" TestNameSuffix="NoZenLoader" Platform="$(Platform)" ProjectName="$(Project)" ProjectTarget="$(ProjectTarget)" NodeName="$(RunNodeLazyLoadName)" Phase="$(Phase)" Iteration="1" Args="$($(Phase)Args)" EditorArgs="$(ProjectCombinedArgs) $(LazyLoadTestArgs) $(NoZenLoaderTestArgs)"/>
 							</Node>
-							<Property Name="RunNodeList" Value="$(RunNodeList);$(RunNodeLazyLoadName)" If="$(EnableLazyLoadTest)"/>
+							<Property Name="RunNodeList" Value="$(RunNodeList);$(RunNodeLazyLoadName)" If="$(EnableLazyLoadNoZenLoaderTest)"/>
 							<!--Zen LazyLoad TestRun if enabled--> 
-							<Property Name="RunNodeLazyLoadZenName" Value="Run $(OutputName) $(Phase) LazyLoad Zen"/>
-							<Node Name="$(RunNodeLazyLoadZenName)" Requires="$(RunNodeRequires);$(RunNodeList)" If="$(EnableLazyLoadZenTest)">
+							<Property Name="RunNodeLazyLoadZenName" Value="Run $(OutputName) $(Phase) LazyLoad ZenLoader"/>
+							<Node Name="$(RunNodeLazyLoadZenName)" Requires="$(RunNodeRequires);$(RunNodeList)" If="$(EnableLazyLoadZenLoaderTest)">
 								<Expand Name="RunTest" TestNameSuffix="ZenLoader" Platform="$(Platform)" ProjectName="$(Project)" ProjectTarget="$(ProjectTarget)" NodeName="$(RunNodeLazyLoadZenName)" Phase="$(Phase)" Iteration="1" Args="$($(Phase)Args)" EditorArgs="$(ProjectCombinedArgs) $(LazyLoadZenTestArgs)"/>
 							</Node>
-							<Property Name="RunNodeList" Value="$(RunNodeList);$(RunNodeLazyLoadZenName)" If="$(EnableLazyLoadZenTest)"/>
+							<Property Name="RunNodeList" Value="$(RunNodeList);$(RunNodeLazyLoadZenName)" If="$(EnableLazyLoadZenLoaderTest)"/>
 						</Do>
 					</ForEach>
 					<Property Name="RunNodeUploadPhaseTraces" Value="Upload $(OutputName) Traces"/>
