Adding variable to allow for additional command line arguments to be passed to compiles on Win64 agents.

#rnx

diff --git a/Engine/Build/InstalledEngineBuild.xml b/Engine/Build/InstalledEngineBuild.xml
--- a/Engine/Build/InstalledEngineBuild.xml
+++ b/Engine/Build/InstalledEngineBuild.xml
@@ -31,6 +31,9 @@
 	<Property Name="DefaultWithIOS" Value="$(DefaultWithPlatform)"/>
 	<Property Name="DefaultWithIOS" Value="false" If="'$(HostPlatform)' != 'Mac' And !$(AllPlatforms)"/>
 
+	<!-- Global Compile Arguments - note: only implemented for Win64 agents -->
+	<Option Name="ExtraCompileArgs" DefaultValue="" Description="Extra arguments to pass to compiles run on Win64 agents"/>
+
 	<!-- Setup default target platforms for Installed build if they haven't been specified on the commandline -->
 	<Option Name="WithWin64" Restrict="true|false" DefaultValue="$(DefaultWithWin64)" Description="Include the Win64 target platform"/>
 	<Option Name="WithMac" Restrict="true|false" DefaultValue="$(DefaultWithMac)" Description="Include the Mac target platform"/>
@@ -167,8 +170,8 @@
 			<SetVersion Change="$(Change)" Branch="$(EscapedBranch)" If="$(IsBuildMachine)"/>
 		</Node>
 		<Node Name="Compile UnrealEditor Win64" Requires="Update Version Files" Produces="#UnrealEditor Win64;#UnrealEditor Win64 Unstripped;#UnrealEditor Win64 Stripped;#UnrealEditor Win64 Unsigned;#UnrealEditor Win64 Signed;#UnrealEditor Win64 WithExceptions">
-			<Compile Target="UnrealEditor" Platform="Win64" Configuration="DebugGame" Tag="#UnrealEditor Win64" Arguments="-precompile -allmodules"/>
-			<Compile Target="UnrealEditor" Platform="Win64" Configuration="Development" Tag="#UnrealEditor Win64" Arguments="-precompile -allmodules"/>
+			<Compile Target="UnrealEditor" Platform="Win64" Configuration="DebugGame" Tag="#UnrealEditor Win64" Arguments="-precompile -allmodules $(ExtraCompileArgs)"/>
+			<Compile Target="UnrealEditor" Platform="Win64" Configuration="Development" Tag="#UnrealEditor Win64" Arguments="-precompile -allmodules $(ExtraCompileArgs)"/>
 
 			<Do If="$(EmbedSrcSrvInfo)">
 				<!-- Embed source info into the PDB files. Should be done from this machine to ensure that paths are correct. -->
@@ -265,8 +268,8 @@
 		<ForEach Name="Target" Values="$(OptionalClientTarget);$(OptionalServerTarget)">
 			<Node Name="Compile $(Target) Win64" Requires="Update Version Files" Produces="#$(Target) Win64">
 				<ForEach Name="Configuration" Values="$(GameConfigurations)">
-					<Compile Target="$(Target)" Platform="Win64" Configuration="$(Configuration)" Tag="#$(Target) Win64" Arguments="-precompile -allmodules -nolink $(TargetDebugInfoArg)"/>
-					<Compile Target="$(Target)" Platform="Win64" Configuration="$(Configuration)" Tag="#$(Target) Win64" Arguments="-precompile $(TargetDebugInfoArg)" Clean="false"/>
+					<Compile Target="$(Target)" Platform="Win64" Configuration="$(Configuration)" Tag="#$(Target) Win64" Arguments="-precompile -allmodules -nolink $(TargetDebugInfoArg) $(ExtraCompileArgs)"/>
+					<Compile Target="$(Target)" Platform="Win64" Configuration="$(Configuration)" Tag="#$(Target) Win64" Arguments="-precompile $(TargetDebugInfoArg) $(ExtraCompileArgs)" Clean="false"/>
 					<Property Name="OptionalTargetTags" Value="$(OptionalTargetTags);#$(Target) Win64"/>
 				</ForEach>
 			</Node>
@@ -276,8 +279,8 @@
 				<Tag Files="$(Tag)" With="#UnrealGame Win64"/>
 			</ForEach>
 			<ForEach Name="Configuration" Values="$(GameConfigurations)">
-				<Compile Target="UnrealGame" Platform="Win64" Configuration="$(Configuration)" Tag="#UnrealGame Win64" AllowParallelExecutor="$(bAllowParallelExecutor)" Arguments="-precompile -allmodules -nolink $(TargetDebugInfoArg)"/>
-				<Compile Target="UnrealGame" Platform="Win64" Configuration="$(Configuration)" Tag="#UnrealGame Win64" AllowParallelExecutor="$(bAllowParallelExecutor)" Arguments="-precompile $(TargetDebugInfoArg)" Clean="false"/>
+				<Compile Target="UnrealGame" Platform="Win64" Configuration="$(Configuration)" Tag="#UnrealGame Win64" AllowParallelExecutor="$(bAllowParallelExecutor)" Arguments="-precompile -allmodules -nolink $(TargetDebugInfoArg) $(ExtraCompileArgs)"/>
+				<Compile Target="UnrealGame" Platform="Win64" Configuration="$(Configuration)" Tag="#UnrealGame Win64" AllowParallelExecutor="$(bAllowParallelExecutor)" Arguments="-precompile $(TargetDebugInfoArg) $(ExtraCompileArgs)" Clean="false"/>
 			</ForEach>
 
 			<Do If="$(EmbedSrcSrvInfo)">
@@ -373,8 +376,8 @@
 		<Node Name="Compile UnrealGame Android" Requires="Update Version Files" Produces="#UnrealGame Android;#UnrealGame Android Unstripped;#UnrealGame Android Stripped">
 			<ForEach Name="Target" Values="UnrealGame;$(OptionalClientTarget)">
 				<ForEach Name="Configuration" Values="$(GameConfigurations)">
-					<Compile Target="$(Target)" Platform="Android" Configuration="$(Configuration)" Tag="#UnrealGame Android" AllowParallelExecutor="$(bAllowParallelExecutor)" Arguments="-precompile -allmodules -nolink -architectures=$(AndroidArchitectures) $(TargetDebugInfoArg)"/>
-					<Compile Target="$(Target)" Platform="Android" Configuration="$(Configuration)" Tag="#UnrealGame Android" AllowParallelExecutor="$(bAllowParallelExecutor)" Arguments="-precompile -architectures=$(AndroidArchitectures) $(TargetDebugInfoArg)" Clean="false"/>
+					<Compile Target="$(Target)" Platform="Android" Configuration="$(Configuration)" Tag="#UnrealGame Android" AllowParallelExecutor="$(bAllowParallelExecutor)" Arguments="-precompile -allmodules -nolink -architectures=$(AndroidArchitectures) $(TargetDebugInfoArg) $(ExtraCompileArgs)"/>
+					<Compile Target="$(Target)" Platform="Android" Configuration="$(Configuration)" Tag="#UnrealGame Android" AllowParallelExecutor="$(bAllowParallelExecutor)" Arguments="-precompile -architectures=$(AndroidArchitectures) $(TargetDebugInfoArg) $(ExtraCompileArgs)" Clean="false"/>
 				</ForEach>
 			</ForEach>
 
@@ -452,9 +455,9 @@
 
 	<Agent Name="NotForLicensee Tools Group Win64" Type="Win64">
 		<Node Name="Build NotForLicensee Tools Win64" Requires="Update Version Files" Produces="#NotForLicensee Build Tools Win64">
-			<Compile Target="DatasmithCADWorker" Configuration="Development" Platform="Win64" Tag="#NotForLicensee Build Tools Win64"/>
+			<Compile Target="DatasmithCADWorker" Configuration="Development" Platform="Win64" Tag="#NotForLicensee Build Tools Win64" Arguments="$(ExtraCompileArgs)"/>
 			<!-- Temporarily removing the compilation of DatasmithCADRuntime from the installed build
-			<Compile Target="DatasmithCADRuntime" Configuration="Development" Platform="Win64" Tag="#NotForLicensee Build Tools Win64"/>
+			<Compile Target="DatasmithCADRuntime" Configuration="Development" Platform="Win64" Tag="#NotForLicensee Build Tools Win64" Arguments="$(ExtraCompileArgs)"/>
 			-->
 		</Node>
 	</Agent>
@@ -471,27 +474,27 @@
 
 	<Agent Name="Tools Group Win64" Type="Win64_Licensee">
 		<Node Name="Build Tools Win64" Requires="Update Version Files;#NotForLicensee Build Tools Win64" Produces="#Tools Win64 Complete">
-			<Compile Target="CrashReportClient" Configuration="Shipping" Platform="Win64" Tag="#Build Tools Win64" Arguments="$(CrashReporterCompileArgs)"/>
-			<Compile Target="CrashReportClientEditor" Configuration="Shipping" Platform="Win64" Tag="#Build Tools Win64" Arguments="$(CrashReporterCompileArgsDev)"/>
+			<Compile Target="CrashReportClient" Configuration="Shipping" Platform="Win64" Tag="#Build Tools Win64" Arguments="$(CrashReporterCompileArgs) ($ExtraCompileArgs)"/>
+			<Compile Target="CrashReportClientEditor" Configuration="Shipping" Platform="Win64" Tag="#Build Tools Win64" Arguments="$(CrashReporterCompileArgsDev) ($ExtraCompileArgs)"/>
 			<Tag Files="#NotForLicensee Build Tools Win64" With="#Build Tools Win64"/>
-			<Compile Target="ShaderCompileWorker" Configuration="Development" Platform="Win64" Tag="#Build Tools Win64"/>
-			<Compile Target="EpicWebHelper" Configuration="Shipping" Platform="Win64" Tag="#Build Tools Win64"/>
-			<Compile Target="UnrealInsights" Configuration="Development" Platform="Win64" Tag="#Build Tools Win64"/>
-			<Compile Target="UnrealFrontend" Configuration="Development" Platform="Win64" Tag="#Build Tools Win64"/>
-			<Compile Target="UnrealLightmass" Configuration="Development" Platform="Win64" Tag="#Build Tools Win64"/>
-			<Compile Target="InterchangeWorker" Configuration="Development" Platform="Win64" Tag="#Build Tools Win64"/>
-			<Compile Target="UnrealPak" Configuration="Development" Platform="Win64" Tag="#Build Tools Win64"/>
-			<Compile Target="UnrealMultiUserServer" Configuration="Development" Platform="Win64" Tag="#Build Tools Win64"/>
-			<Compile Target="UnrealMultiUserSlateServer" Configuration="Development" Platform="Win64" Tag="#Build Tools Win64"/>
-			<Compile Target="UnrealRecoverySvc" Configuration="Development" Platform="Win64" Tag="#Build Tools Win64"/>
-			<Compile Target="LiveCodingConsole" Configuration="Development" Platform="Win64" Tag="#Build Tools Win64"/>
-			<Compile Target="BootstrapPackagedGame" Configuration="Shipping" Platform="Win64" Tag="#Build Tools Win64"/>
-			<Compile Target="BuildPatchTool" Configuration="Shipping" Platform="Win64" Tag="#Build Tools Win64"/>
-			<Compile Target="SwitchboardListener" Configuration="Development" Platform="Win64" Tag="#Build Tools Win64"/>
-			<Compile Target="SwitchboardListenerHelper" Configuration="Development" Platform="Win64" Tag="#Build Tools Win64"/>
-			<Compile Target="UnrealObjectPtrTool" Configuration="Development" Platform="Win64" Tag="#Build Tools Win64"/>
-			<Compile Target="ZenDashboard" Configuration="Development" Platform="Win64" Tag="#Build Tools Win64"/>
-			<Compile Target="LiveLinkHub" Configuration="Shipping" Platform="Win64" Arguments="-Monolithic" Tag="#Build Tools Win64"/>
+			<Compile Target="ShaderCompileWorker" Configuration="Development" Platform="Win64" Tag="#Build Tools Win64" Arguments="$(ExtraCompileArgs)"/>
+			<Compile Target="EpicWebHelper" Configuration="Shipping" Platform="Win64" Tag="#Build Tools Win64" Arguments="$(ExtraCompileArgs)"/>
+			<Compile Target="UnrealInsights" Configuration="Development" Platform="Win64" Tag="#Build Tools Win64" Arguments="$(ExtraCompileArgs)"/>
+			<Compile Target="UnrealFrontend" Configuration="Development" Platform="Win64" Tag="#Build Tools Win64" Arguments="$(ExtraCompileArgs)"/>
+			<Compile Target="UnrealLightmass" Configuration="Development" Platform="Win64" Tag="#Build Tools Win64" Arguments="$(ExtraCompileArgs)"/>
+			<Compile Target="InterchangeWorker" Configuration="Development" Platform="Win64" Tag="#Build Tools Win64" Arguments="$(ExtraCompileArgs)"/>
+			<Compile Target="UnrealPak" Configuration="Development" Platform="Win64" Tag="#Build Tools Win64" Arguments="$(ExtraCompileArgs)"/>
+			<Compile Target="UnrealMultiUserServer" Configuration="Development" Platform="Win64" Tag="#Build Tools Win64" Arguments="$(ExtraCompileArgs)"/>
+			<Compile Target="UnrealMultiUserSlateServer" Configuration="Development" Platform="Win64" Tag="#Build Tools Win64" Arguments="$(ExtraCompileArgs)"/>
+			<Compile Target="UnrealRecoverySvc" Configuration="Development" Platform="Win64" Tag="#Build Tools Win64" Arguments="$(ExtraCompileArgs)"/>
+			<Compile Target="LiveCodingConsole" Configuration="Development" Platform="Win64" Tag="#Build Tools Win64" Arguments="$(ExtraCompileArgs)"/>
+			<Compile Target="BootstrapPackagedGame" Configuration="Shipping" Platform="Win64" Tag="#Build Tools Win64" Arguments="$(ExtraCompileArgs)"/>
+			<Compile Target="BuildPatchTool" Configuration="Shipping" Platform="Win64" Tag="#Build Tools Win64" Arguments="$(ExtraCompileArgs)"/>
+			<Compile Target="SwitchboardListener" Configuration="Development" Platform="Win64" Tag="#Build Tools Win64" Arguments="$(ExtraCompileArgs)"/>
+			<Compile Target="SwitchboardListenerHelper" Configuration="Development" Platform="Win64" Tag="#Build Tools Win64" Arguments="$(ExtraCompileArgs)"/>
+			<Compile Target="UnrealObjectPtrTool" Configuration="Development" Platform="Win64" Tag="#Build Tools Win64" Arguments="$(ExtraCompileArgs)"/>
+			<Compile Target="ZenDashboard" Configuration="Development" Platform="Win64" Tag="#Build Tools Win64" Arguments="$(ExtraCompileArgs)"/>
+			<Compile Target="LiveLinkHub" Configuration="Shipping" Platform="Win64" Arguments="-Monolithic $(ExtraCompileArgs)" Tag="#Build Tools Win64"/>
 
 			<Tag Files="Engine/Binaries/Win64/XGEControlWorker.exe" With="#Build Tools Win64"/>
 
@@ -932,7 +935,7 @@
 			<Do If="$(WithWin64) And $(CompileDatasmithPlugins) == 'True'">
 				<Property Name="DatasmithPlugins" Value="DatasmithCADImporter;DatasmithC4DImporter;AxFImporter;MDLImporter"/>
 				<ForEach Name="Plugin" Values="$(DatasmithPlugins)">
-					<Compile Target="UnrealEditor" Configuration="Development" Platform="Win64" Arguments="-Plugin=&quot;$(RootDir)/Engine/Plugins/Enterprise/$(Plugin)/$(Plugin).uplugin&quot;"/>
+					<Compile Target="UnrealEditor" Configuration="Development" Platform="Win64" Arguments="-Plugin=&quot;$(RootDir)/Engine/Plugins/Enterprise/$(Plugin)/$(Plugin).uplugin&quot; $(ExtraCompileArgs)"/>
 					<Tag Files="$(RootDir)/Engine/Plugins/Enterprise/$(Plugin)/Binaries/Win64/..." With="#DatasmithPluginBinaries Win64"/>
 				</ForEach>
 			</Do>
