Add CRC_DEFAULT_URL and CRC_DEFAULT_COMPANY_NAME to installed build script.

diff --git a/Engine/Build/InstalledEngineBuild.xml b/Engine/Build/InstalledEngineBuild.xml
--- a/Engine/Build/InstalledEngineBuild.xml
+++ b/Engine/Build/InstalledEngineBuild.xml
@@ -150,17 +150,27 @@
 	<!-- Include the script that lists files and patterns used when creating an installed build -->
 	<Include Script="InstalledEngineFilters.xml"/>
 
-	<!-- Crash reporter telemetry settings for the customer facing crash reporter -->
-	<Option Name="CrashReporterAPIURL" DefaultValue="" Description="The URL to use to talk to the CrashReporterClient API. Used for crashes from packed game."/>
-	<Option Name="CrashReporterAPIKey" DefaultValue="" Description="The API key to use to talk to the CrashReporterClient API. Used for crashes from packaged game."/>
-	<Property Name="CrashReporterCompileArgs" Value=""/>
-	<Property Name="CrashReporterCompileArgs" Value="-define:CRC_TELEMETRY_URL=\&quot;$(CrashReporterAPIURL)\&quot; -define:CRC_TELEMETRY_KEY_DEV=\&quot;$(CrashReporterAPIKey)\&quot; -define:CRC_TELEMETRY_KEY_RELEASE=\&quot;$(CrashReporterAPIKey)\&quot; -OverrideBuildEnvironment" If="'$(CrashReporterAPIURL)' != '' And '$(CrashReporterAPIKey)' != ''" />
+	<!-- Crash reporter and telemetry settings for the customer facing crash reporter -->
+	<Option Name="CrashReporterURL" DefaultValue="" Description="The URL to use to talk to the crash reporting API. Used for crashes from packaged games."/>
+	<Option Name="CrashReporterTelemetryURL" DefaultValue="" Description="The URL to use to talk to the Telemetry API. Used for crashes from packaged games."/>
+	<Option Name="CrashReporterTelemetryKey" DefaultValue="" Description="The API key to use to talk to the Telemetry API. Used for crashes from packaged games."/>
+	<Option Name="CrashReporterCompanyName" DefaultValue="" Description="Default company name to display in CrashReportClient. Used for crashes from packaged games."/>
+
+	<Property Name="CRCTelemetryDefines" Value=""/>
+	<Property Name="CRCTelemetryDefines" Value="-define:CRC_TELEMETRY_URL=\&quot;$(CrashReporterTelemetryURL)\&quot; -define:CRC_TELEMETRY_KEY=\&quot;$(CrashReporterTelemetryKey)\&quot; -OverrideBuildEnvironment" If="'$(CrashReporterTelemetryURL)' != '' And '$(CrashReporterTelemetryKey)' != ''" />
+	<Property Name="CRCDefines" Value=""/>
+	<Property Name="CRCDefines" Value=" -define:CRC_DEFAULT_URL=\&quot;$(CrashReporterURL)\&quot; -define:CRC_DEFAULT_COMPANY_NAME=\&quot;$(CrashReporterCompanyName)\&quot; -OverrideBuildEnvironment" If="'$(CrashReporterURL)' != ''" />
 	
-	<!-- Crash reporter telemetry setting for development time crash reporter -->
-	<Option Name="CrashReporterDevAPIURL" DefaultValue="" Description="The URL to use to talk to the CrashReporterClient API. Used for development time crashes."/>
-	<Option Name="CrashReporterDevAPIKey" DefaultValue="" Description="The API key to use to talk to the CrashReporterClient API. Used for development time crashes."/>
-	<Property Name="CrashReporterCompileArgsDev" Value=""/>
-	<Property Name="CrashReporterCompileArgsDev" Value="-define:CRC_TELEMETRY_URL=\&quot;$(CrashReporterDevAPIURL)\&quot; -define:CRC_TELEMETRY_KEY_DEV=\&quot;$(CrashReporterDevAPIKey)\&quot; -define:CRC_TELEMETRY_KEY_RELEASE=\&quot;$(CrashReporterDevAPIKey)\&quot; -OverrideBuildEnvironment" If="'$(CrashReporterDevAPIURL)' != '' And '$(CrashReporterDevAPIKey)' != ''" />
+	<!-- Crash reporter and telemetry settings for development time crash reporter -->
+	<Option Name="CrashReporterDevURL" DefaultValue="" Description="The URL to use to talk to the crash reporting API. Used for development time crashes."/>
+	<Option Name="CrashReporterDevTelemetryURL" DefaultValue="" Description="The URL to use to talk to the Telemetry API. Used for development time crashes."/>
+	<Option Name="CrashReporterDevTelemetryKey" DefaultValue="" Description="The API key to use to talk to the Telemetry API. Used for development time crashes."/>
+	<Option Name="CrashReporterDevCompanyName" DefaultValue="" Description="Default company name to display in CrashReportClient. Used for development time crashes."/>
+
+	<Property Name="CRCETelemetryDefines" Value=""/>
+	<Property Name="CRCETelemetryDefines" Value="-define:CRC_TELEMETRY_URL=\&quot;$(CrashReporterDevTelemetryURL)\&quot; -define:CRC_TELEMETRY_KEY=\&quot;$(CrashReporterDevTelemetryKey)\&quot; -OverrideBuildEnvironment" If="'$(CrashReporterDevTelemetryURL)' != '' And '$(CrashReporterDevTelemetryKey)' != ''" />
+	<Property Name="CRCEDefines" Value=""/>
+	<Property Name="CRCEDefines" Value=" -define:CRC_DEFAULT_URL=\&quot;$(CrashReporterDevURL)\&quot; -define:CRC_DEFAULT_COMPANY_NAME=\&quot;$(CrashReporterDevCompanyName)\&quot; -OverrideBuildEnvironment" If="'$(CrashReporterDevURL)' != ''" />
 
 	<Property Name="MacArchitectureArg" Value="-architecture=arm64+x64" />
 
@@ -474,8 +484,8 @@
 
 	<Agent Name="Tools Group Win64" Type="Win64_Licensee">
 		<Node Name="Build Tools Win64" Requires="Update Version Files;#NotForLicensee Build Tools Win64" Produces="#Tools Win64 Complete">
-			<Compile Target="CrashReportClient" Configuration="Shipping" Platform="Win64" Tag="#Build Tools Win64" Arguments="$(CrashReporterCompileArgs) ($ExtraCompileArgs)"/>
-			<Compile Target="CrashReportClientEditor" Configuration="Shipping" Platform="Win64" Tag="#Build Tools Win64" Arguments="$(CrashReporterCompileArgsDev) ($ExtraCompileArgs)"/>
+			<Compile Target="CrashReportClient" Configuration="Shipping" Platform="Win64" Tag="#Build Tools Win64" Arguments="$(CRCTelemetryDefines) $(CRCDefines) $(ExtraCompileArgs)"/>
+			<Compile Target="CrashReportClientEditor" Configuration="Shipping" Platform="Win64" Tag="#Build Tools Win64" Arguments="$(CRCETelemetryDefines) $(CRCEDefines) $(ExtraCompileArgs)"/>
 			<Tag Files="#NotForLicensee Build Tools Win64" With="#Build Tools Win64"/>
 			<Compile Target="ShaderCompileWorker" Configuration="Development" Platform="Win64" Tag="#Build Tools Win64" Arguments="$(ExtraCompileArgs)"/>
 			<Compile Target="EpicWebHelper" Configuration="Shipping" Platform="Win64" Tag="#Build Tools Win64" Arguments="$(ExtraCompileArgs)"/>
@@ -560,8 +570,8 @@
 			<Sign Files="#Sign Binaries" Parallel="$(SignWindowsExecutablesInParallel)" If="'$(SignExecutables)' == true"/>
 		</Node>
 		<Node Name="Build Tools Linux" Requires="Update Version Files">
-			<Compile Target="CrashReportClient" Platform="Linux" Configuration="Shipping" Tag="#Build Tools Linux" Arguments="$(CrashReporterCompileArgs)"/>
-			<Compile Target="CrashReportClientEditor" Platform="Linux" Configuration="Shipping" Tag="#Build Tools Linux" Arguments="$(CrashReporterCompileArgsDev)"/>
+			<Compile Target="CrashReportClient" Platform="Linux" Configuration="Shipping" Tag="#Build Tools Linux" Arguments="$(CRCTelemetryDefines)  $(CRCDefines)"/>
+			<Compile Target="CrashReportClientEditor" Platform="Linux" Configuration="Shipping" Tag="#Build Tools Linux" Arguments="$(CRCETelemetryDefines)  $(CRCEDefines)"/>
 			<Compile Target="ShaderCompileWorker" Configuration="Development" Platform="Linux" Tag="#Build Tools Linux"/>
 			<Compile Target="EpicWebHelper" Configuration="Shipping" Platform="Linux" Tag="#Build Tools Linux"/>
 			<Compile Target="UnrealInsights" Configuration="Development" Platform="Linux" Tag="#Build Tools Linux"/>
@@ -630,8 +640,8 @@
 		</Node>
 	
 		<Node Name="Build Tools Mac" Requires="Update Version Files">
-			<Compile Target="CrashReportClient" Configuration="Shipping" Platform="Mac" Tag="#Build Tools Mac" Arguments="-EnableDSYM $(CrashReporterCompileArgs)"/>
-			<Compile Target="CrashReportClientEditor" Configuration="Shipping" Platform="Mac" Tag="#Build Tools Mac" Arguments="-EnableDSYM $(CrashReporterCompileArgsDev)"/>
+			<Compile Target="CrashReportClient" Configuration="Shipping" Platform="Mac" Tag="#Build Tools Mac" Arguments="-EnableDSYM $(CRCTelemetryDefines) $(CRCDefines)"/>
+			<Compile Target="CrashReportClientEditor" Configuration="Shipping" Platform="Mac" Tag="#Build Tools Mac" Arguments="-EnableDSYM $(CRCETelemetryDefines) $(CRCEDefines)"/>
 			<Compile Target="ShaderCompileWorker" Configuration="Development" Platform="Mac" Tag="#Build Tools Mac" Arguments="-EnableDSYM"/>
 			<Compile Target="EpicWebHelper" Configuration="Shipping" Platform="Mac" Tag="#Build Tools Mac" Arguments="-EnableDSYM"/>
 			<Compile Target="UnrealInsights" Configuration="Development" Platform="Mac" Tag="#Build Tools Mac" Arguments="-EnableDSYM"/>
@@ -654,8 +664,8 @@
 
 	<Agent Name="Tools Group LinuxNative" Type="Linux_Licensee">
 		<Node Name="Build Tools LinuxNative" Requires="Update Version Files">
-			<Compile Target="CrashReportClient" Platform="Linux" Configuration="Shipping" Tag="#Build Tools LinuxNative" Arguments="$(CrashReporterCompileArgs)"/>
-			<Compile Target="CrashReportClientEditor" Platform="Linux" Configuration="Shipping"  Tag="#Build Tools LinuxNative" Arguments="$(CrashReporterCompileArgsDev)"/>
+			<Compile Target="CrashReportClient" Platform="Linux" Configuration="Shipping" Tag="#Build Tools LinuxNative" Arguments="$(CRCTelemetryDefines) $(CRCDefines)"/>
+			<Compile Target="CrashReportClientEditor" Platform="Linux" Configuration="Shipping"  Tag="#Build Tools LinuxNative" Arguments="$(CRCETelemetryDefines) $(CRCEDefines)"/>
 			<Compile Target="ShaderCompileWorker" Configuration="Development" Platform="Linux" Tag="#Build Tools LinuxNative"/>
 			<Compile Target="EpicWebHelper" Configuration="Shipping" Platform="Linux" Tag="#Build Tools LinuxNative"/>
 			<Compile Target="UnrealInsights" Configuration="Development" Platform="Linux" Tag="#Build Tools LinuxNative"/>
