Horde: Add steps to upload the installed engine build to Horde once built.

diff --git a/Engine/Build/InstalledEngineBuild.xml b/Engine/Build/InstalledEngineBuild.xml
--- a/Engine/Build/InstalledEngineBuild.xml
+++ b/Engine/Build/InstalledEngineBuild.xml
@@ -1132,6 +1132,9 @@
 			<Tag Files="#Installed Build Win64 Files" Filter="*.target" With="#Installed Build Win64 Receipts"/>
 			<SanitizeReceipt Files="#Installed Build Win64 Receipts" EngineDir="$(LocalInstalledDir)\Engine"/>
 		</Node>
+		<Node Name="Upload To Horde Win64" Requires="Make Installed Build Win64">
+			<CreateArtifact Name="installed-build-win64" Type="installed-build" BaseDir="$(LocalInstalledDir)"/>
+		</Node>
 	</Agent>
 	<Aggregate Name="HostPlatforms_Win64" Label="Builds/Win64" Requires="Make Installed Build Win64" Exclude="Update Version Files"/>
 
@@ -1256,6 +1259,9 @@
 			<Tag Files="#Installed Build Mac Files" Filter="*.target" With="#Installed Build Mac Receipts"/>
 			<SanitizeReceipt Files="#Installed Build Mac Receipts" EngineDir="$(LocalInstalledDirMac)/Engine"/>
 		</Node>
+		<Node Name="Upload To Horde Mac" Requires="Make Installed Build Mac">
+			<CreateArtifact Name="installed-build-mac" Type="installed-build" BaseDir="$(LocalInstalledDirMac)"/>
+		</Node>
 	</Agent>
 	<Aggregate Name="HostPlatforms_Mac" Label="Builds/Mac" Requires="Make Installed Build Mac"/>
 
@@ -1371,7 +1377,16 @@
 			<Tag Files="#Installed Build Linux Files" Filter="*.target" With="#Installed Build Linux Receipts"/>
 			<SanitizeReceipt Files="#Installed Build Linux Receipts" EngineDir="$(LocalInstalledDirLinux)/Engine"/>
 		</Node>
+		<Node Name="Upload To Horde Linux" Requires="Make Installed Build Linux">
+			<CreateArtifact Name="installed-build-linux" Type="installed-build" BaseDir="$(LocalInstalledDirLinux)"/>
+		</Node>
 	</Agent>
 	<Aggregate Name="HostPlatforms_Linux" Label="Builds/Linux" Requires="Make Installed Build Linux"/>
 
+	<Property Name="InstalledEngineBuildTargets"/>
+	<Property Name="InstalledEngineBuildTargets" Value="$(InstalledEngineBuildTargets)Upload to Horde Win64;" If="'$(WithWin64)' == true"/>
+	<Property Name="InstalledEngineBuildTargets" Value="$(InstalledEngineBuildTargets)Upload to Horde Mac;" If="'$(WithMac)' == true"/>
+	<Property Name="InstalledEngineBuildTargets" Value="$(InstalledEngineBuildTargets)Upload to Horde Linux;" If="'$(WithLinux)' == true"/>
+	<Aggregate Name="Installed Engine Build" Requires="$(InstalledEngineBuildTargets)"/>
+
 </BuildGraph>
diff --git a/Engine/Source/Programs/Horde/Horde.Server/Defaults/ue5-dev-main-hordedemo.stream.json b/Engine/Source/Programs/Horde/Horde.Server/Defaults/ue5-dev-main-hordedemo.stream.json
--- a/Engine/Source/Programs/Horde/Horde.Server/Defaults/ue5-dev-main-hordedemo.stream.json
+++ b/Engine/Source/Programs/Horde/Horde.Server/Defaults/ue5-dev-main-hordedemo.stream.json
@@ -418,8 +418,8 @@
 			"name": "Installed Engine Build",
 			"initialAgentType": "Win64",
 			"arguments": [
-				"-Target=UE5 Binary Release",
-				"-Script=Engine/Build/InternalEngineBuild.xml"
+				"-Target=Installed Engine Build",
+				"-Script=Engine/Build/InstalledEngineBuild.xml"
 			],
 			"parameters": [
 				{
