Horde: Move options for host platforms in installed build into main build script.

diff --git a/Engine/Build/InstalledEngineBuild.xml b/Engine/Build/InstalledEngineBuild.xml
--- a/Engine/Build/InstalledEngineBuild.xml
+++ b/Engine/Build/InstalledEngineBuild.xml
@@ -34,6 +34,11 @@
 	<!-- Global Compile Arguments - note: only implemented for Win64 agents -->
 	<Option Name="ExtraCompileArgs" DefaultValue="" Description="Extra arguments to pass to compiles run on Win64 agents"/>
 
+	<!-- Host platforms for the installed editor -->
+	<Option Name="WithInstalledWin64" Restrict="true|false" DefaultValue="true" Description="Whether to include the Win64 editor as part of a binary release"/>
+	<Option Name="WithInstalledMac" Restrict="true|false" DefaultValue="true" Description="Whether to include the Mac editor as part of a binary release"/>
+	<Option Name="WithInstalledLinux" Restrict="true|false" DefaultValue="true" Description="Whether to include the Linux editor as part of a binary release"/>
+
 	<!-- Setup default target platforms for Installed build if they haven't been specified on the commandline -->
 	<Option Name="WithWin64" Restrict="true|false" DefaultValue="$(DefaultWithWin64)" Description="Include the Win64 target platform"/>
 	<Option Name="WithMac" Restrict="true|false" DefaultValue="$(DefaultWithMac)" Description="Include the Mac target platform"/>
@@ -1384,9 +1389,9 @@
 	<Aggregate Name="HostPlatforms_Linux" Label="Builds/Linux" Requires="Make Installed Build Linux"/>
 
 	<Property Name="InstalledEngineBuildTargets"/>
-	<Property Name="InstalledEngineBuildTargets" Value="$(InstalledEngineBuildTargets)Upload to Horde Win64;" If="'$(WithWin64)' == true"/>
-	<Property Name="InstalledEngineBuildTargets" Value="$(InstalledEngineBuildTargets)Upload to Horde Mac;" If="'$(WithMac)' == true"/>
-	<Property Name="InstalledEngineBuildTargets" Value="$(InstalledEngineBuildTargets)Upload to Horde Linux;" If="'$(WithLinux)' == true"/>
+	<Property Name="InstalledEngineBuildTargets" Value="$(InstalledEngineBuildTargets)Upload to Horde Win64;" If="'$(WithInstalledWin64)' == true"/>
+	<Property Name="InstalledEngineBuildTargets" Value="$(InstalledEngineBuildTargets)Upload to Horde Mac;" If="'$(WithInstalledMac)' == true"/>
+	<Property Name="InstalledEngineBuildTargets" Value="$(InstalledEngineBuildTargets)Upload to Horde Linux;" If="'$(WithInstalledLinux)' == true"/>
 	<Aggregate Name="Installed Engine Build" Requires="$(InstalledEngineBuildTargets)"/>
 
 </BuildGraph>
