Fixed build break: Removed plugins not compiling on Mac

#rnx

diff --git a/Engine/Build/InstalledEngineBuild.xml b/Engine/Build/InstalledEngineBuild.xml
--- a/Engine/Build/InstalledEngineBuild.xml
+++ b/Engine/Build/InstalledEngineBuild.xml
@@ -960,7 +960,7 @@
 	<Agent Name="Extra Plugins Agent Mac" Type="Mac">
 		<Node Name="Compile Datasmith Plugins Mac" Requires="Update Version Files" Produces="#DatasmithPluginBinaries Mac">
 			<Do If="$(WithMac) And $(CompileDatasmithPlugins) == 'True'">
-				<Property Name="DatasmithPlugins" Value="DatasmithCADImporter;DatasmithC4DImporter;AxFImporter;MDLImporter"/>
+				<Property Name="DatasmithPlugins" Value="DatasmithC4DImporter;MDLImporter"/>
 				<ForEach Name="Plugin" Values="$(DatasmithPlugins)">
 					<Compile Target="UnrealEditor" Configuration="Development" Platform="Mac" Arguments="-Plugin=&quot;$(RootDir)/Engine/Plugins/Enterprise/$(Plugin)/$(Plugin).uplugin&quot; $(MacArchitectureArg)"/>
 					<Tag Files="$(RootDir)/Engine/Plugins/Enterprise/$(Plugin)/Binaries/Mac/..." With="#DatasmithPluginBinaries Mac"/>
