Ensure that engine installed builds don't fail if the PDB files are missing for zen.exe and zenserver.exe on Windows.

diff --git a/Engine/Build/InstalledEngineBuild.xml b/Engine/Build/InstalledEngineBuild.xml
--- a/Engine/Build/InstalledEngineBuild.xml
+++ b/Engine/Build/InstalledEngineBuild.xml
@@ -834,10 +834,8 @@
 			<Tag Files="Engine\Binaries\DotNET\Ionic.Zip.Reduced.dll" With="#ToCopy"/>
 			<Tag Files="Engine\Binaries\DotNET\OneSky.dll" With="#ToCopy"/>
 			<Tag Files="Engine\Binaries\Win64\crashpad_handler.exe" With="#ToCopy"/>
-			<Tag Files="Engine\Binaries\Win64\zen.exe" With="#ToCopy"/>
-			<Tag Files="Engine\Binaries\Win64\zen.pdb" With="#ToCopy"/>
-			<Tag Files="Engine\Binaries\Win64\zenserver.exe" With="#ToCopy"/>
-			<Tag Files="Engine\Binaries\Win64\zenserver.pdb" With="#ToCopy"/>
+			<Tag Files="Engine\Binaries\Win64\zen.*" With="#ToCopy"/>
+			<Tag Files="Engine\Binaries\Win64\zenserver.*" With="#ToCopy"/>
 			<Tag Files="Templates\TemplateResources\..." With="#ToCopy"/>
 			<Expand Name="TagDDCEngineFiles" ExceptPlugins="Plugins\...\Binaries\Mac\..."/>
 			<!-- Filter out the files not needed to build DDC. Removing confidential folders can affect DDC keys, so we want to be sure that we're making DDC with a build that can use it. -->
