Add missing iOS simulator libs to IB

diff --git a/Engine/Build/InstalledEngineFilters.xml b/Engine/Build/InstalledEngineFilters.xml
--- a/Engine/Build/InstalledEngineFilters.xml
+++ b/Engine/Build/InstalledEngineFilters.xml
@@ -582,7 +582,20 @@
 		<!-- PLCrashReporter needs to be specificed manually as there are two versions of the library -->
 		Engine/Source/ThirdParty/PLCrashReporter/lib/lib-Xcode-11.3.1/iOS/...
 		Engine/Source/ThirdParty/PLCrashReporter/lib/lib-Xcode-12.4/iOS/...
+
+		<!-- Extra iOS Simualtor libraries -->
+		Engine/Plugins/Compression/OodleNetwork/Sdks/2.9.12/lib/IOS/....sim.a
+		Engine/Source/ThirdParty/BLAKE3/1.3.1/lib/IOS/Simulator/...
+		Engine/Source/ThirdParty/ICU/icu4c-64_1/lib/IOS/Simulator/...
+		Engine/Source/ThirdParty/FreeType2/FreeType2-2.10.0/lib/IOS/Simulator/...
+		Engine/Source/ThirdParty/libPNG/libPNG-1.5.2/lib/IOS/Simulator/...
+		Engine/Source/ThirdParty/HarfBuzz/harfbuzz-2.4.0/lib/IOS/Simulator/...
+		Engine/Source/ThirdParty/Ogg/libogg-1.2.2/lib/ios/libogg.sim.a
+		Engine/Source/ThirdParty/Vorbis/libvorbis-1.3.2/lib/IOS/....sim.a 
+		Engine/Source/ThirdParty/libOpus/opus-1.1/IOS/libOpus.sim.a
+		Engine/Source/ThirdParty/OpenSSL/1.1.1t/lib/IOS/....sim.a
 	</Property>
+
 	<Property Name="CopyIOSFilterWin64" Value="$(CopyIOSFilter)">
 		Engine/Binaries/DotNET/IOS/openssl.exe
 		Engine/Binaries/Mac/DsymExporter*
