Expanding support for custom gauntlet tests in EditorPerf.xml.

diff --git a/Engine/Build/EditorPerf.xml b/Engine/Build/EditorPerf.xml
--- a/Engine/Build/EditorPerf.xml
+++ b/Engine/Build/EditorPerf.xml
@@ -82,6 +82,15 @@
 	<!-- Iterations to run and average together -->
 	<Option Name="Iterations" DefaultValue="1" Description="The list of iteration numbers to run. For example '1;2;3' for three iterations"/>
 
+	<!-- Option for number of PIE iterations to run (per editor launch) -->
+	<Property Name="DefaultNumIterationsPie" Value="2" />
+	<Option Name="NumIterationsPIE" DefaultValue="$(DefaultNumIterationsPIE)" Description="Specify the number of times to start/stop PIE sessions in a single editor run."/>
+	<Do If="'$(NumIterationsPIE)' &lt;= '0'">
+		<Warning Message="NumIterationPIE must be greater than 0. Setting to default of $(DefaultNumIterationsPIE)."/>
+		<Property Name="NumIterationsPIE" Value="$(DefaultNumIterationsPIE)" />
+	</Do>
+	<Property Name="ArgIterationsPIE" Value="-ini:Engine:[ConsoleVariables]:project.iteration.pie.testiterations=$(NumIterationsPIE)" />
+
 	<!-- Used to determine if we will include internal projects -->
 	<Option Name="IsLicensee" Restrict="true|false" DefaultValue="false" Description="Define to true if this script is used externally."/>
 
@@ -89,7 +98,7 @@
 	<Property Name="AllProjects" Value=""/>
 
 	<!-- A macro that can be used to easily set up a new test project. -->
-	<Macro Name="SetupPerfProject" Arguments="ProjectName" OptionalArguments="CopyDefaultsFrom;ProjectTarget;EditorTarget;MapPath;GameSpecificArgs;UseBinariesFrom;ExtraNonCopyableArgs">
+	<Macro Name="SetupPerfProject" Arguments="ProjectName" OptionalArguments="CopyDefaultsFrom;ProjectTarget;EditorTarget;MapPath;GameSpecificArgs;UseBinariesFrom;ExtraNonCopyableArgs;GauntletTest;GauntletTestArgs">
 		<Property Name="$(ProjectName)__ProjectTarget" Value="$(ProjectName)" CreateInParentScope="true"/>
 		<Property Name="$(ProjectName)__ProjectTarget" Value="$($(CopyDefaultsFrom)__ProjectTarget)" If="'$(CopyDefaultsFrom)' != ''"/>
 		<Property Name="$(ProjectName)__ProjectTarget" Value="$(ProjectTarget)" If="'$(ProjectTarget)' != ''"/>
@@ -98,9 +107,9 @@
 		<Property Name="$(ProjectName)__EditorTarget" Value="$($(CopyDefaultsFrom)__EditorTarget)" If="'$(CopyDefaultsFrom)' != ''"/>
 		<Property Name="$(ProjectName)__EditorTarget" Value="$(EditorTarget)" If="'$(EditorTarget)' != ''"/>
 
-		<Property Name="$(ProjectName)__MapArgs" Value="" CreateInParentScope="true"/>
-		<Property Name="$(ProjectName)__MapArgs" Value="$($(CopyDefaultsFrom)__MapArgs)" If="'$(CopyDefaultsFrom)' != ''"/>
-		<Property Name="$(ProjectName)__MapArgs" Value="-ini:Engine:[/Script/Engine.AutomationTestSettings]:MapsToPIETest=\&quot;$(MapPath)\&quot;" If="'$(MapPath)' != ''"/>
+		<!-- Stash the MapPath so that we can use it later via CopyDefaultsFrom and construct new MapArgs -->
+		<Property Name="$(ProjectName)__BaseMapPath" Value="$(MapPath)" CreateInParentScope="true"/>
+		<Property Name="$(ProjectName)__BaseMapPath" Value="$($(CopyDefaultsFrom)__BaseMapPath)" If="'$(CopyDefaultsFrom)' != ''"/>
 
 		<Property Name="$(ProjectName)__SpecificArgs" Value="" CreateInParentScope="true"/>
 		<Property Name="$(ProjectName)__SpecificArgs" Value="$($(CopyDefaultsFrom)__SpecificArgs)" If="'$(CopyDefaultsFrom)' != ''"/>
@@ -110,6 +119,18 @@
 		<Property Name="$(ProjectName)__UseBinariesFrom" Value="$(CopyDefaultsFrom)" If="'$(CopyDefaultsFrom)' != ''"/>
 		<Property Name="$(ProjectName)__UseBinariesFrom" Value="$(UseBinariesFrom)" If="'$(UseBinariesFrom)' != ''"/>
 
+		<Property Name="$(ProjectName)__GauntletTest" Value="EditorTests" CreateInParentScope="true"/>
+		<Property Name="$(ProjectName)__GauntletTest" Value="$($(CopyDefaultsFrom)__GauntletTest)" If="'$(CopyDefaultsFrom)' != ''"/>
+		<Property Name="$(ProjectName)__GauntletTest" Value="$(GauntletTest)" If="'$(GauntletTest)' != ''"/>
+
+		<Property Name="$(ProjectName)__GauntletTestArgs" Value="" CreateInParentScope="true"/>
+		<Property Name="$(ProjectName)__GauntletTestArgs" Value="$($(CopyDefaultsFrom)__GauntletTestArgs)" If="'$(CopyDefaultsFrom)' != ''"/>
+		<Property Name="$(ProjectName)__GauntletTestArgs" Value="$(GauntletTestArgs)" If="'$(GauntletTestArgs)' != ''"/>
+
+		<!-- MapArgs is an "indirect" parameter. We don't ever copy it around, instead we always construct it based on BaseMapPath and others parameters -->
+		<Property Name="$(ProjectName)__MapArgs" Value="" CreateInParentScope="true"/>
+		<Property Name="$(ProjectName)__MapArgs" Value="-ini:Engine:[/Script/Engine.AutomationTestSettings]:MapsToPIETest=\&quot;$($(ProjectName)__BaseMapPath)\&quot;" If="'$($(ProjectName)__BaseMapPath)' != '' and '$($(ProjectName)__GauntletTest)'=='EditorTests'" />
+
 		<!-- Immediately resolve any indirections on UseBinariesFrom. We cannot support them later when we set up the actual nodes. -->
 		<Do If="'$($(ProjectName)__UseBinariesFrom)' != ''">
 			<Property Name="$(ProjectName)__UseBinariesFrom" Value="$($($(ProjectName)__UseBinariesFrom)__UseBinariesFrom)" If="'$($($(ProjectName)__UseBinariesFrom)__UseBinariesFrom)' != ''"/>
@@ -135,15 +156,6 @@
 	<Option Name="EnableZenLoaderAsyncThreadTest" DefaultValue="false" Description="Enable additional test per phase that uses the Zen Loader with an async loading thread."/>
 	<Option Name="DisableEditorDomainForZenLoaderTests" DefaultValue="true" Description="Explicitly disables editordomain (-noeditordomain) for zen loader tests"/>
 
-	<!-- Option for number of PIE iterations to run (per editor launch) -->
-	<Property Name="DefaultNumIterationsPie" Value="2" />
-	<Option Name="NumIterationsPIE" DefaultValue="$(DefaultNumIterationsPIE)" Description="Specify the number of times to start/stop PIE sessions in a single editor run."/>
-	<Do If="'$(NumIterationsPIE)' &lt;= '0'">
-		<Warning Message="NumIterationPIE must be greater than 0. Setting to default of $(DefaultNumIterationsPIE)."/>
-		<Property Name="NumIterationsPIE" Value="$(DefaultNumIterationsPIE)" />
-	</Do>
-	<Property Name="ArgIterationsPIE" Value="-ini:Engine:[ConsoleVariables]:project.iteration.pie.testiterations=$(NumIterationsPIE)" />
-	
 	<!-- Zen Loader tests specifc arguments -->
 	<Property Name="ZenLoaderTestArgs" Value="-zenloader"/>
 	<Property Name="NoZenLoaderTestArgs" Value="-nozenloader"/>
@@ -175,7 +187,7 @@
 	<Property Name="GauntletArgs" Value="$(GauntletArgs) -maxduration=28800"/>
 
 	<!-- Gauntlet arguments for profile runs -->
-	<Property Name="ProfileArgs" Value="-test=EditorTests -tracefiletrunc"/>
+	<Property Name="ProfileArgs" Value="-tracefiletrunc"/>
 	<Property Name="ProfileArgs" Value="$(ProfileArgs) -EnablePlugins=$(ExtraPluginsCommaDelims)" If="'$(ExtraPluginsCommaDelims)' != ''"/>
 	<Property Name="ProfileArgs" Value="$(ProfileArgs) -SuppressLogWarnings -SuppressLogErrors"/>
 	<Property Name="ProfileArgs" Value="$(ProfileArgs) -GameInstanceLostTimerSeconds=3600"/>
@@ -201,8 +213,19 @@
 	<!-- The main test running macro logic that is expanded with platform and dependency node names -->
 	<Macro Name="RunTest" Arguments="TestNameSuffix;Platform;ProjectName;ProjectTarget;NodeName;Phase;Iteration;Args;EditorArgs">
 		<Property Name="TestName" Value="EditorPerf$(Phase)$(TestNameSuffix)"/>
-		<Property Name="TestFilter" Value="Project.Iteration.PIE"/>
-		<Property Name="TestFilterArgs" Value="-runtest=$(TestFilter)"/>
+		<Property Name="AssetRegistryCacheRootFolderArg" Value="-AssetRegistryCacheRootFolder=&quot;$(AssetRegistryCache_Folder)&quot;"/>
+
+		<Property Name="TestFilterArgs" Value=""/>
+		<Property Name="CombinedEditorArgs" Value=""/>
+		<Do If="'$($(ProjectName)__GauntletTest)'=='EditorTests'">
+			<Property Name="TestFilter" Value="Project.Iteration.PIE"/>
+			<Property Name="TestFilterArgs" Value="-runtest=$(TestFilter)"/>
+			<Property Name="CombinedEditorArgs" Value="$(EditorArgs) $(TestFilterArgs) $(AssetRegistryCacheRootFolderArg) $(ArgIterationsPIE)"/>
+		</Do>
+		<Do If="'$($(ProjectName)__GauntletTest)'!='EditorTests'">
+			<Property Name="CombinedEditorArgs" Value="$(EditorArgs) $(AssetRegistryCacheRootFolderArg)"/>
+		</Do>
+
 		<Property Name="IterationSuffix" Value=""/>
 		<Property Name="IterationSuffix" Value="$(Iteration)" If="'$(Iteration)' != '1'"/>
 
@@ -216,18 +239,19 @@
 		<Property Name="LogOutputParams" Value="-logdir=&quot;$(TestOutputDir)&quot;"/>
 		<Property Name="ReportParams" Value="-WriteTestResultsForHorde -HordeTestDataKey=&quot;$(NodeName)$(PreflightSuffix)&quot;"/>
 		<Property Name="ArtifactName" Value="$(Phase)$(IterationSuffix)$(TestNameSuffix)"/>
-		<Property Name="AllGauntletArgs" Value="-Project=$(ProjectTarget) -Platform=$(Platform) $(GauntletArgs) $(ReportParams) $(LogOutputParams) -artifactname=&quot;$(ArtifactName)&quot;"/>
-		<Property Name="AssetRegistryCacheRootFolderArg" Value="-AssetRegistryCacheRootFolder=&quot;$(AssetRegistryCache_Folder)&quot;"/>
-		<Property Name="CombinedEditorArgs" Value="$(EditorArgs) $(AssetRegistryCacheRootFolderArg) $(ArgIterationsPIE)"/>
-		<Property Name="RunArgs" Value="$(Args) -editorargs=&quot;$(CombinedEditorArgs)&quot;"/>
+		<Property Name="AllGauntletArgs" Value="-Project=$(ProjectTarget) -Platform=$(Platform) $(GauntletArgs) $($(Project)__GauntletTestArgs) $(ReportParams) $(LogOutputParams) -artifactname=&quot;$(ArtifactName)&quot;"/>
+
+		<Property Name="TestArg" Value="-test=$($(Project)__GauntletTest)" />
+			
+		<Property Name="RunArgs" Value="$(TestArg) $(TestFilterArgs) $(Args) -editorargs=&quot;$(CombinedEditorArgs)&quot;"/>
 
 		<!-- Do DDC cleanup outside the editor process, handles Cold cleanup without changing stats -->
 		<!-- <Commandlet Project="$(ProjectTarget)" Name="DDCCleanup" Arguments="$(Args) $(CommandletArgs)"/> -->
 
 		<!-- Run the Gauntlet EditorPerf Tests -->
-		<Command Name="RunUnreal" Arguments="$(AllGauntletArgs) $(ProfileArgs) $(TestFilterArgs) $(RunArgs) -sessionlabel=&quot;$(TestName)&quot; -tracefile=&quot;$(TraceFilePath)&quot;"/>
+		<Command Name="RunUnreal" Arguments="$(AllGauntletArgs) $(ProfileArgs) $(RunArgs) -sessionlabel=&quot;$(TestName)&quot; -tracefile=&quot;$(TraceFilePath)&quot;"/>
 
-		<!-- Translate the utrace file into csv data -->
+		<!-- Translate the utrace file into csv data --> 
 		<Do If="$(RunTraceSummary)">
 			<Commandlet Project="$(ProjectTarget)" Name="SummarizeTrace" Arguments="-inputfile=&quot;$(TraceFilePath)&quot; -testname=&quot;$(TestName)&quot; $(SummarizeArgs) $(CommandletArgs)"/>
 		</Do>
