Fix livelink hub plugins not getting staged with UE installed build

#prelfight 660341070f0b188ba61d2aec

diff --git a/Engine/Build/InstalledEngineBuild.xml b/Engine/Build/InstalledEngineBuild.xml
--- a/Engine/Build/InstalledEngineBuild.xml
+++ b/Engine/Build/InstalledEngineBuild.xml
@@ -72,13 +72,7 @@
 
 	<!--- Create command line argument for build identifier -->
 	<Property Name="BuildIdentifierArg" Value="" />
-	<Property Name="BuildIdentifierArg" Value="-BuildId=&quot;$(BuildId)&quot;" If="'$(BuildId)' != ''"/>
-	
-	<!-- LiveLinkHub args -->
-	<Option Name="WithLiveLinkHub3rdPartyPlugins" DefaultValue="true" Restrict="true|false" Description="Whether or not to compile LiveLinkHub plugins."/>
-	<Property Name="LiveLinkHubCompileArgs" Value="-EnableThirdPartyPlugins=true" />
-	<Property Name="LiveLinkHubCompileArgs" Value="-EnableThirdPartyPlugins=false" If="!$(WithLiveLinkHub3rdPartyPlugins)"/>
-	
+	<Property Name="BuildIdentifierArg" Value="-BuildId=&quot;$(BuildId)&quot;" If="'$(BuildId)' != ''"/>	
 
 	<!-- Whether to embed source server info (ie. P4 paths) into PDB files -->
 	<Option Name="EmbedSrcSrvInfo" Restrict="true|false" DefaultValue="false" Description="Whether to add Source indexing to Windows game apps so they can be added to a symbol server"/>
@@ -114,8 +108,14 @@
 	<Property Name="AndroidArchitectures" Value="arm64+x64"/>
 	
 	<Property Name="DefaultArchitectureCommandsWin64" Value=""/>
-
-
+	
+	<!-- LiveLinkHub args -->
+	<Option Name="WithLiveLinkHub3rdPartyPlugins" DefaultValue="true" Restrict="true|false" Description="Whether or not to compile LiveLinkHub plugins."/>
+	<Property Name="LiveLinkHubCompileArgs" Value="-EnableThirdPartyPlugins=true" />
+	<Property Name="LiveLinkHubCompileArgs" Value="-EnableThirdPartyPlugins=false" If="!$(WithLiveLinkHub3rdPartyPlugins)"/>
+	<Property Name="StageLiveLinkHubArgs" Value="" />
+	<Property Name="StageLiveLinkHubArgs" Value="$(StageLiveLinkHubArgs) -NoSign" If="!$(SignExecutables)" />
+	
 	<!-- Compile flags for Windows Game binaries -->
 	<Property Name="TargetDebugInfoArg" Value="-nodebuginfo"/>
 	<Property Name="TargetDebugInfoArg" Value="" If="$(WithFullDebugInfo)"/>
@@ -475,12 +475,42 @@
 
 	<Agent Name="NotForLicensee Tools Group Win64" Type="Win64">
 		<Node Name="Build NotForLicensee Tools Win64" Requires="Update Version Files" Produces="#NotForLicensee Build Tools Win64">
+			<Compile Target="LiveLinkHub" Configuration="Shipping" Platform="Win64" Arguments="-Monolithic $(LiveLinkHubCompileArgs) $(ExtraCompileArgs)" Tag="#NotForLicensee Build Tools Win64" />
+			<Command Name="StageLiveLinkHub" Arguments="$(StageLiveLinkHubArgs) -Platform=Win64"/>
+			
+			<Tag Files="Engine/Binaries/Win64/ViconDataStreamSDK_CPP.dll" With="#NotForLicensee Build Tools Win64" />
+			<Tag Files="Engine/Binaries/Win64/boost_thread-vc141-mt-x64-1_75.dll" With="#NotForLicensee Build Tools Win64" />
+			<Tag Files="Engine/Programs/LiveLinkHub/Staged/..." With="#NotForLicensee Build Tools Win64" />
+			
 			<Compile Target="DatasmithCADWorker" Configuration="Development" Platform="Win64" Tag="#NotForLicensee Build Tools Win64" Arguments="$(ExtraCompileArgs)"/>
 			<!-- Temporarily removing the compilation of DatasmithCADRuntime from the installed build
 			<Compile Target="DatasmithCADRuntime" Configuration="Development" Platform="Win64" Tag="#NotForLicensee Build Tools Win64" Arguments="$(ExtraCompileArgs)"/>
 			-->
 		</Node>
+		
+		<Node Name="Build NotForLicensee Tools Linux" Requires="Update Version Files" Produces="#NotForLicensee Build Tools Linux">
+			<Compile Target="LiveLinkHub" Configuration="Shipping" Platform="Linux" Tag="#NotForLicensee Build Tools Linux" Arguments="-Monolithic $(LiveLinkHubCompileArgs)"/>
+			<Command Name="StageLiveLinkHub" Arguments="$(StageLiveLinkHubArgs) -Platform=Linux"/>
+			<!-- Since we're running on the same agent as the Win64 node, we don't need to re-tag those files or else it would generate an error. -->
+		</Node>
+	</Agent>
+	
+	<Agent Name="NotForLicensee Tools Group LinuxNative" Type="Linux">
+		<Node Name="Build NotForLicensee Tools LinuxNative" Requires="Update Version Files" Produces="#NotForLicensee Build Tools LinuxNative">
+			<Compile Target="LiveLinkHub" Configuration="Shipping" Platform="Linux" Tag="#NotForLicensee Build Tools LinuxNative" Arguments="-Monolithic $(LiveLinkHubCompileArgs)"/>
+			<Command Name="StageLiveLinkHub" Arguments="$(StageLiveLinkHubArgs) -Platform=Linux"/>
+			<Tag Files="Engine/Programs/LiveLinkHub/Staged/..." With="#NotForLicensee Build Tools LinuxNative" />
+		</Node>
+	</Agent>
+	
+	<Agent Name="NotForLicensee Tools Group Mac" Type="Mac">
+		<Node Name="Build NotForLicensee Tools Mac" Requires="Update Version Files" Produces="#NotForLicensee Build Tools Mac">
+			<Compile Target="LiveLinkHub" Configuration="Shipping" Platform="Mac" Tag="#NotForLicensee Build Tools Mac" Arguments="-Monolithic $(LiveLinkHubCompileArgs)"/>
+			<Command Name="StageLiveLinkHub" Arguments="$(StageLiveLinkHubArgs) -Platform=Mac"/>
+			<Tag Files="Engine/Programs/LiveLinkHub/Staged/..." With="#NotForLicensee Build Tools Mac" />
+		</Node>
 	</Agent>
+
 	
 	<Agent Name="BuildRules" Type="Win64_Licensee">
 		<Node Name="Compile Rules Assemblies" Produces="#Rules Assemblies">
@@ -514,7 +544,6 @@
 			<Compile Target="SwitchboardListenerHelper" Configuration="Development" Platform="Win64" Tag="#Build Tools Win64" Arguments="$(ExtraCompileArgs)"/>
 			<Compile Target="UnrealObjectPtrTool" Configuration="Development" Platform="Win64" Tag="#Build Tools Win64" Arguments="$(ExtraCompileArgs)"/>
 			<Compile Target="ZenDashboard" Configuration="Development" Platform="Win64" Tag="#Build Tools Win64" Arguments="$(ExtraCompileArgs)"/>
-			<Compile Target="LiveLinkHub" Configuration="Shipping" Platform="Win64" Arguments="-Monolithic $(LiveLinkHubCompileArgs) $(ExtraCompileArgs)" Tag="#Build Tools Win64"/>
 
 			<Tag Files="Engine/Binaries/Win64/XGEControlWorker.exe" With="#Build Tools Win64"/>
 
@@ -579,7 +608,8 @@
 			<!-- Sign the binaries -->
 			<Sign Files="#Sign Binaries" Parallel="$(SignWindowsExecutablesInParallel)" If="'$(SignExecutables)' == true"/>
 		</Node>
-		<Node Name="Build Tools Linux" Requires="Update Version Files">
+		<Node Name="Build Tools Linux" Requires="Update Version Files;#NotForLicensee Build Tools Linux" >
+			<Tag Files="#NotForLicensee Build Tools Linux" With="#Build Tools Linux"/>
 			<Compile Target="CrashReportClient" Platform="Linux" Configuration="Shipping" Tag="#Build Tools Linux" Arguments="$(CRCTelemetryDefines)  $(CRCDefines)"/>
 			<Compile Target="CrashReportClientEditor" Platform="Linux" Configuration="Shipping" Tag="#Build Tools Linux" Arguments="$(CRCETelemetryDefines)  $(CRCEDefines)"/>
 			<Compile Target="ShaderCompileWorker" Configuration="Development" Platform="Linux" Tag="#Build Tools Linux"/>
@@ -649,7 +679,8 @@
 			<Sign Files="#Sign Binaries Mac" If="'$(SignExecutables)' == true"/>
 		</Node>
 	
-		<Node Name="Build Tools Mac" Requires="Update Version Files">
+		<Node Name="Build Tools Mac" Requires="Update Version Files;#NotForLicensee Build Tools Mac">
+			<Tag Files="#NotForLicensee Build Tools Mac" With="#Build Tools Mac"/>
 			<Compile Target="CrashReportClient" Configuration="Shipping" Platform="Mac" Tag="#Build Tools Mac" Arguments="-EnableDSYM $(CRCTelemetryDefines) $(CRCDefines)"/>
 			<Compile Target="CrashReportClientEditor" Configuration="Shipping" Platform="Mac" Tag="#Build Tools Mac" Arguments="-EnableDSYM $(CRCETelemetryDefines) $(CRCEDefines)"/>
 			<Compile Target="ShaderCompileWorker" Configuration="Development" Platform="Mac" Tag="#Build Tools Mac" Arguments="-EnableDSYM"/>
@@ -668,12 +699,12 @@
 			<Compile Target="UnrealAtoS" Configuration="Development" Platform="Mac" Tag="#Build Tools Mac" Arguments="-EnableDSYM"/>
 			<Compile Target="UnrealObjectPtrTool" Configuration="Development" Platform="Mac" Tag="#Build Tools Mac" Arguments="-EnableDSYM"/>
 			<Compile Target="ZenDashboard" Configuration="Development" Platform="Mac" Tag="#Build Tools Mac" Arguments="-EnableDSYM"/>
-			<Compile Target="LiveLinkHub" Configuration="Shipping" Platform="Mac" Tag="#Build Tools Mac" Arguments="-Monolithic $(LiveLinkHubCompileArgs)"/>
 		</Node>
 	</Agent>
 
 	<Agent Name="Tools Group LinuxNative" Type="Linux_Licensee">
-		<Node Name="Build Tools LinuxNative" Requires="Update Version Files">
+		<Node Name="Build Tools LinuxNative" Requires="Update Version Files;#NotForLicensee Build Tools LinuxNative">
+			<Tag Files="#NotForLicensee Build Tools LinuxNative" With="#Build Tools LinuxNative"/>
 			<Compile Target="CrashReportClient" Platform="Linux" Configuration="Shipping" Tag="#Build Tools LinuxNative" Arguments="$(CRCTelemetryDefines) $(CRCDefines)"/>
 			<Compile Target="CrashReportClientEditor" Platform="Linux" Configuration="Shipping"  Tag="#Build Tools LinuxNative" Arguments="$(CRCETelemetryDefines) $(CRCEDefines)"/>
 			<Compile Target="ShaderCompileWorker" Configuration="Development" Platform="Linux" Tag="#Build Tools LinuxNative"/>
@@ -690,7 +721,6 @@
 			<Compile Target="SwitchboardListenerHelper" Configuration="Development" Platform="Linux" Tag="#Build Tools LinuxNative"/>
 			<Compile Target="UnrealObjectPtrTool" Configuration="Development" Platform="Linux" Tag="#Build Tools LinuxNative"/>
 			<Compile Target="ZenDashboard" Configuration="Development" Platform="Linux" Tag="#Build Tools LinuxNative"/>
-			<Compile Target="LiveLinkHub" Configuration="Shipping" Platform="Linux" Tag="#Build Tools LinuxNative" Arguments="-Monolithic $(LiveLinkHubCompileArgs)"/>
 		</Node>
 
 		<!-- Compile the C# tools. Copies files to a temporary directory first, to strip out all confidential folders and avoid sharing violations with the running version of UAT. -->
diff --git a/Engine/Build/InstalledEngineFilters.xml b/Engine/Build/InstalledEngineFilters.xml
--- a/Engine/Build/InstalledEngineFilters.xml
+++ b/Engine/Build/InstalledEngineFilters.xml
@@ -326,11 +326,21 @@
 
 		<!-- nanosvg -->
 		Engine/Source/ThirdParty/nanosvg/...
+		
+		<!-- LiveLinkHub files -->
+		Engine/Binaries/Win64/LiveLinkHub-Win64-Shipping.exe
+		Engine/Binaries/Win64/ViconDataStreamSDK_CPP.dll
+		Engine/Binaries/Win64/boost_thread-vc141-mt-x64-1_75.dll
+		Engine/Programs/LiveLinkHub/...
 	</Property>
 
 	<Property Name="CopyEditorExceptionsWin64" Value="$(CopyEditorExceptions)">
 		<!-- Exclude Mac binaries on windows here, because previous wildcard for plugin resources matches OSX dsym resources on Windows -->
 		Engine/Plugins/.../Binaries/Mac/...
+		
+		Engine/Binaries/Win64/LiveLinkHub-Win64-Shipping.pdb
+		Engine/Binaries/Linux/LiveLinkHub-Linux-Shipping.debug
+		Engine/Binaries/Linux/LiveLinkHub-Linux-Shipping.sym
 	</Property>
 
 	<!-- Mac editor specific dependencies -->
@@ -376,6 +386,10 @@
 		Engine/Binaries/Mac/zenserver
 		Engine/Binaries/Mac/zen
 		Engine/Binaries/Mac/crashpad_handler
+		
+		<!-- LiveLinkHub files -->
+		Engine/Binaries/Mac/LiveLinkHub-Mac-Shipping.app/...
+		Engine/Programs/LiveLinkHub/...
 	</Property>
 
 	<Property Name="CopyEditorExceptionsMac" Value="$(CopyEditorExceptions)">
@@ -437,6 +451,11 @@
 		<!-- zenserver is the tool for Unreal Zen Storage services like Local DDC and optionally cooked output storage -->
 		Engine/Binaries/Linux/zenserver
 		Engine/Binaries/Linux/zen
+		
+		<!-- LiveLinkHub Files -->
+		Engine/Binaries/Linux/LiveLinkHub-Linux-Shipping
+		Engine/Programs/LiveLinkHub/...
+		
 	</Property>
 	<Property Name="CopyEditorExceptionsLinux" Value="$(CopyEditorExceptions)">
 		<!-- Exclude Windows/Mac binaries on Linux here -->
@@ -450,6 +469,10 @@
 		<!-- Don't want The Windows/Mac Python packages -->
 		Engine/Plugins/Experimental/PythonFoundationPackages/Content/Python/Lib/Win64/...
 		Engine/Plugins/Experimental/PythonFoundationPackages/Content/Python/Lib/Mac/...
+		
+		<!-- We remove these in an attempt to save space. -->
+		Engine/Binaries/Linux/LiveLinkHub-Linux-Shipping.debug
+		Engine/Binaries/Linux/LiveLinkHub-Linux-Shipping.sym
 	</Property>
 
 	<!-- Target Platform Filters/Exceptions -->
diff --git a/Engine/Source/Programs/LiveLinkHub/Private/LiveLinkHubRun.cpp b/Engine/Source/Programs/LiveLinkHub/Private/LiveLinkHubRun.cpp
--- a/Engine/Source/Programs/LiveLinkHub/Private/LiveLinkHubRun.cpp
+++ b/Engine/Source/Programs/LiveLinkHub/Private/LiveLinkHubRun.cpp
@@ -32,19 +32,21 @@ int32 RunLiveLinkHub(const TCHAR* CommandLine)
 		});
 #endif
 
-#if COOKED_EDITOR_DISTRIBUTION
-	/** If LLH runs in a cooked editor distribution, then we must mount our RemappedEngineFiles folder to the /Engine/ folder in order to
-	 * 	have access to the uncooked content necessary for starting LiveLinkHub.
-	 */
-	FCommandLine::Set(CommandLine);
-
-	// Must match the path in LiveLinkHubStaging.cs
+	// Must match the path in LiveLinkHubStaging.cs / StageLiveLinkHub.automation.cs
 	const FString RemappedFolder = FPaths::EngineDir() / TEXT("Programs/LiveLinkHub/Staged");
 
 	const FString RemappedConfigs = RemappedFolder / TEXT("Config");
 	const FString RemappedContent = RemappedFolder / TEXT("Content");
 	const FString RemappedPlugins = RemappedFolder / TEXT("Plugins");
 
+	// Needs to be initialized early for mount points / plugin search paths.
+	FCommandLine::Set(CommandLine);
+
+#if COOKED_EDITOR_DISTRIBUTION
+
+	/** If LLH runs in a cooked editor distribution, then we must mount our RemappedEngineFiles folder to the /Engine/ folder in order to
+	 * 	have access to the uncooked content necessary for starting LiveLinkHub.
+	 */
 	if (!IFileManager::Get().DirectoryExists(*RemappedContent))
 	{
 		IFileManager::Get().MakeDirectory(*RemappedContent, true);
@@ -69,11 +71,13 @@ int32 RunLiveLinkHub(const TCHAR* CommandLine)
 	// This is needed to ensure that content generated by the hub will be discoverable (ie. LiveLink recordings)
 	const FString GameContentPath = FPaths::Combine(FPlatformProcess::UserSettingsDir(), *FApp::GetEpicProductIdentifier(), TEXT("LiveLinkHub"), TEXT("Content"));
 	FPackageName::RegisterMountPoint(TEXT("/Game/"), GameContentPath);
-	
+#endif
+
 	// Needed to avoid assertion in scoped slow task
 	GWarn = FPlatformApplicationMisc::GetFeedbackContext();
+
+	// Add staged plugins to the search path so LLH can find the manifest and content.
 	IPluginManager::Get().AddPluginSearchPath(RemappedPlugins, /*bShouldRefresh*/ true);
-#endif
 
 	// Disable this when going through PreInit to prevent the console window from appearing.
 	GIsSilent = true;
