CIS warning fix
Change the LiveLinkHub compile step to not tag the livelinkhub target receipt, which would generate warnings later on since the restricted build dependencies could not be found.

diff --git a/Engine/Build/InstalledEngineBuild.xml b/Engine/Build/InstalledEngineBuild.xml
--- a/Engine/Build/InstalledEngineBuild.xml
+++ b/Engine/Build/InstalledEngineBuild.xml
@@ -475,9 +475,10 @@
 
 	<Agent Name="NotForLicensee Tools Group Win64" Type="Win64">
 		<Node Name="Build NotForLicensee Tools Win64" Requires="Update Version Files" Produces="#NotForLicensee Build Tools Win64">
-			<Compile Target="LiveLinkHub" Configuration="Shipping" Platform="Win64" Arguments="-Monolithic $(LiveLinkHubCompileArgs) $(ExtraCompileArgs)" Tag="#NotForLicensee Build Tools Win64" />
+			<Compile Target="LiveLinkHub" Configuration="Shipping" Platform="Win64" Arguments="-Monolithic $(LiveLinkHubCompileArgs) $(ExtraCompileArgs)" />
 			<Command Name="StageLiveLinkHub" Arguments="$(StageLiveLinkHubArgs) -Platform=Win64"/>
 			
+			<Tag Files="Engine/Binaries/Win64/LiveLinkHub-Win64-Shipping.exe" With="#NotForLicensee Build Tools Win64" />
 			<Tag Files="Engine/Binaries/Win64/ViconDataStreamSDK_CPP.dll" With="#NotForLicensee Build Tools Win64" />
 			<Tag Files="Engine/Binaries/Win64/boost_thread-vc141-mt-x64-1_75.dll" With="#NotForLicensee Build Tools Win64" />
 			<Tag Files="Engine/Programs/LiveLinkHub/Staged/..." With="#NotForLicensee Build Tools Win64" />
@@ -489,25 +490,29 @@
 		</Node>
 		
 		<Node Name="Build NotForLicensee Tools Linux" Requires="Update Version Files" Produces="#NotForLicensee Build Tools Linux">
-			<Compile Target="LiveLinkHub" Configuration="Shipping" Platform="Linux" Tag="#NotForLicensee Build Tools Linux" Arguments="-Monolithic $(LiveLinkHubCompileArgs)"/>
+			<Compile Target="LiveLinkHub" Configuration="Shipping" Platform="Linux" Arguments="-Monolithic $(LiveLinkHubCompileArgs)"/>
 			<Command Name="StageLiveLinkHub" Arguments="$(StageLiveLinkHubArgs) -Platform=Linux"/>
+			<Tag Files="Engine/Binaries/Linux/LiveLinkHub-Linux-Shipping" With="#NotForLicensee Build Tools Linux" />
+			
 			<!-- Since we're running on the same agent as the Win64 node, we don't need to re-tag those files or else it would generate an error. -->
 		</Node>
 	</Agent>
 	
 	<Agent Name="NotForLicensee Tools Group LinuxNative" Type="Linux">
 		<Node Name="Build NotForLicensee Tools LinuxNative" Requires="Update Version Files" Produces="#NotForLicensee Build Tools LinuxNative">
-			<Compile Target="LiveLinkHub" Configuration="Shipping" Platform="Linux" Tag="#NotForLicensee Build Tools LinuxNative" Arguments="-Monolithic $(LiveLinkHubCompileArgs)"/>
+			<Compile Target="LiveLinkHub" Configuration="Shipping" Platform="Linux" Arguments="-Monolithic $(LiveLinkHubCompileArgs)"/>
 			<Command Name="StageLiveLinkHub" Arguments="$(StageLiveLinkHubArgs) -Platform=Linux"/>
 			<Tag Files="Engine/Programs/LiveLinkHub/Staged/..." With="#NotForLicensee Build Tools LinuxNative" />
+			<Tag Files="Engine/Binaries/Linux/LiveLinkHub-Linux-Shipping" With="#NotForLicensee Build Tools LinuxNative" />
 		</Node>
 	</Agent>
 	
 	<Agent Name="NotForLicensee Tools Group Mac" Type="Mac">
 		<Node Name="Build NotForLicensee Tools Mac" Requires="Update Version Files" Produces="#NotForLicensee Build Tools Mac">
-			<Compile Target="LiveLinkHub" Configuration="Shipping" Platform="Mac" Tag="#NotForLicensee Build Tools Mac" Arguments="-Monolithic $(LiveLinkHubCompileArgs)"/>
+			<Compile Target="LiveLinkHub" Configuration="Shipping" Platform="Mac" Arguments="-Monolithic $(LiveLinkHubCompileArgs)"/>
 			<Command Name="StageLiveLinkHub" Arguments="$(StageLiveLinkHubArgs) -Platform=Mac"/>
 			<Tag Files="Engine/Programs/LiveLinkHub/Staged/..." With="#NotForLicensee Build Tools Mac" />
+			<Tag Files="Engine/Binaries/Mac/LiveLinkHub-Mac-Shipping/..." With="#NotForLicensee Build Tools Mac" />
 		</Node>
 	</Agent>
 
