Fixing LiveLinkHub issues with InstalledEngineBuild.xml

#rnx

diff --git a/Engine/Build/InstalledEngineBuild.xml b/Engine/Build/InstalledEngineBuild.xml
--- a/Engine/Build/InstalledEngineBuild.xml
+++ b/Engine/Build/InstalledEngineBuild.xml
@@ -111,10 +111,9 @@
 	
 	<!-- LiveLinkHub args -->
 	<Option Name="WithLiveLinkHub3rdPartyPlugins" DefaultValue="true" Restrict="true|false" Description="Whether or not to compile LiveLinkHub plugins."/>
+	<Option Name="HasExtraNotForLicenseeFiles" DefaultValue="false" Restrict="true|false" Description="Whether or not a an extra set of nodes should be generated."/>
 	<Property Name="LiveLinkHubCompileArgs" Value="-EnableThirdPartyPlugins=true" />
 	<Property Name="LiveLinkHubCompileArgs" Value="-EnableThirdPartyPlugins=false" If="!$(WithLiveLinkHub3rdPartyPlugins)"/>
-	<Property Name="StageLiveLinkHubArgs" Value="" />
-	<Property Name="StageLiveLinkHubArgs" Value="$(StageLiveLinkHubArgs) -NoSign" If="!$(SignExecutables)" />
 	
 	<!-- Compile flags for Windows Game binaries -->
 	<Property Name="TargetDebugInfoArg" Value="-nodebuginfo"/>
@@ -473,46 +472,50 @@
 	</Agent>
 	<Aggregate Name="TargetPlatforms_TVOS" Label="Platforms/TVOS" Requires="Compile UnrealGame TVOS"/>
 
+	<Do If="!$(HasExtraNotForLicenseeFiles)">
+		<!-- If this is false we stub out the macro. -->
+		<Macro Name="ExtraNotForLicenseeFiles" Arguments="TagSet;Platform;HostPlatform" />
+	</Do>
+
 	<Agent Name="NotForLicensee Tools Group Win64" Type="Win64">
 		<Node Name="Build NotForLicensee Tools Win64" Requires="Update Version Files" Produces="#NotForLicensee Build Tools Win64">
-			<Compile Target="LiveLinkHub" Configuration="Shipping" Platform="Win64" Arguments="-Monolithic $(LiveLinkHubCompileArgs) $(ExtraCompileArgs)" />
-			<Command Name="StageLiveLinkHub" Arguments="$(StageLiveLinkHubArgs) -Platform=Win64"/>
-			
-			<Tag Files="Engine/Binaries/Win64/LiveLinkHub-Win64-Shipping.exe" With="#NotForLicensee Build Tools Win64" />
-			<Tag Files="Engine/Binaries/Win64/ViconDataStreamSDK_CPP.dll" With="#NotForLicensee Build Tools Win64" />
-			<Tag Files="Engine/Binaries/Win64/boost_thread-vc141-mt-x64-1_75.dll" With="#NotForLicensee Build Tools Win64" />
-			<Tag Files="Engine/Programs/LiveLinkHub/..." With="#NotForLicensee Build Tools Win64" />
-			
 			<Compile Target="DatasmithCADWorker" Configuration="Development" Platform="Win64" Tag="#NotForLicensee Build Tools Win64" Arguments="$(ExtraCompileArgs)"/>
 			<!-- Temporarily removing the compilation of DatasmithCADRuntime from the installed build
 			<Compile Target="DatasmithCADRuntime" Configuration="Development" Platform="Win64" Tag="#NotForLicensee Build Tools Win64" Arguments="$(ExtraCompileArgs)"/>
 			-->
+			<Compile Target="LiveLinkHub" Configuration="Shipping" Platform="Win64" Arguments="-Monolithic $(LiveLinkHubCompileArgs) $(ExtraCompileArgs)" />
+
+			<Expand Name="ExtraNotForLicenseeFiles" TagSet="#NotForLicensee Build Tools Win64" Platform="Win64" HostPlatform="Win64"/>
+			<Tag Files="Engine/Binaries/Win64/LiveLinkHub-Win64-Shipping.exe" With="#NotForLicensee Build Tools Win64"/>
+			<Tag Files="Engine/Programs/LiveLinkHub/..." With="#NotForLicensee Build Tools Win64" />
 		</Node>
 		
 		<Node Name="Build NotForLicensee Tools Linux" Requires="Update Version Files" Produces="#NotForLicensee Build Tools Linux">
 			<Compile Target="LiveLinkHub" Configuration="Shipping" Platform="Linux" Arguments="-Monolithic $(LiveLinkHubCompileArgs)"/>
-			<Command Name="StageLiveLinkHub" Arguments="$(StageLiveLinkHubArgs) -Platform=Linux"/>
-			<Tag Files="Engine/Binaries/Linux/LiveLinkHub-Linux-Shipping" With="#NotForLicensee Build Tools Linux" />
-			
+
 			<!-- Since we're running on the same agent as the Win64 node, we don't need to re-tag those files or else it would generate an error. -->
+			<Expand Name="ExtraNotForLicenseeFiles" TagSet="#NotForLicensee Build Tools Linux" Platform="Linux" HostPlatform="Linux"/>
+			<Tag Files="Engine/Binaries/Linux/LiveLinkHub-Linux-Shipping" With="#NotForLicensee Build Tools Linux" />
 		</Node>
 	</Agent>
 	
 	<Agent Name="NotForLicensee Tools Group LinuxNative" Type="Linux">
 		<Node Name="Build NotForLicensee Tools LinuxNative" Requires="Update Version Files" Produces="#NotForLicensee Build Tools LinuxNative">
 			<Compile Target="LiveLinkHub" Configuration="Shipping" Platform="Linux" Arguments="-Monolithic $(LiveLinkHubCompileArgs)"/>
-			<Command Name="StageLiveLinkHub" Arguments="$(StageLiveLinkHubArgs) -Platform=Linux"/>
-			<Tag Files="Engine/Programs/LiveLinkHub/..." With="#NotForLicensee Build Tools LinuxNative" />
+
+			<Expand Name="ExtraNotForLicenseeFiles" TagSet="#NotForLicensee Build Tools LinuxNative" Platform="LinuxNative" HostPlatform="Linux"/>
 			<Tag Files="Engine/Binaries/Linux/LiveLinkHub-Linux-Shipping" With="#NotForLicensee Build Tools LinuxNative" />
+			<Tag Files="Engine/Programs/LiveLinkHub/..." With="#NotForLicensee Build Tools LinuxNative" />
 		</Node>
 	</Agent>
 	
 	<Agent Name="NotForLicensee Tools Group Mac" Type="Mac">
 		<Node Name="Build NotForLicensee Tools Mac" Requires="Update Version Files" Produces="#NotForLicensee Build Tools Mac">
 			<Compile Target="LiveLinkHub" Configuration="Shipping" Platform="Mac" Arguments="-Monolithic $(LiveLinkHubCompileArgs)"/>
-			<Command Name="StageLiveLinkHub" Arguments="$(StageLiveLinkHubArgs) -Platform=Mac"/>
+
+			<Expand Name="ExtraNotForLicenseeFiles" TagSet="#NotForLicensee Build Tools Mac" Platform="Mac" HostPlatform="Mac"/>
+			<Tag Files="Engine/Binaries/Mac/LiveLinkHub-Mac-Shipping.app/..." With="#NotForLicensee Build Tools Mac" />
 			<Tag Files="Engine/Programs/LiveLinkHub/..." With="#NotForLicensee Build Tools Mac" />
-			<Tag Files="Engine/Binaries/Mac/LiveLinkHub-Mac-Shipping/..." With="#NotForLicensee Build Tools Mac" />
 		</Node>
 	</Agent>
 
