Fixes EditorWorkflowTests in CI. The `-test=EngineTest` was removed as the default value for a property used in various macros as part of refactoring in 32524766. However not all macros were updated to still receive a parameter dropping the -test param from workflow tests

diff --git a/Engine/Build/EditorPerf.xml b/Engine/Build/EditorPerf.xml
--- a/Engine/Build/EditorPerf.xml
+++ b/Engine/Build/EditorPerf.xml
@@ -242,10 +242,9 @@
 		<Property Name="LogOutputParams" Value="-logdir=&quot;$(TestOutputDir)&quot;"/>
 		<Property Name="ReportParams" Value="-WriteTestResultsForHorde -HordeTestDataKey=&quot;$(NodeName)$(PreflightSuffix)&quot;"/>
 		<Property Name="ArtifactName" Value="$(Phase)$(IterationSuffix)$(TestNameSuffix)"/>
-		<Property Name="AllGauntletArgs" Value="-Project=$(ProjectTarget) -Platform=$(Platform) $(GauntletArgs) $($(Project)__GauntletTestArgs) $(ReportParams) $(LogOutputParams) -artifactname=&quot;$(ArtifactName)&quot;"/>
+		<Property Name="AllGauntletArgs" Value="-Project=$(ProjectTarget) -Platform=$(Platform) $(GauntletArgs) $($(ProjectName)__GauntletTestArgs) $(ReportParams) $(LogOutputParams) -artifactname=&quot;$(ArtifactName)&quot;"/>
 
-		<Property Name="TestArg" Value="-test=$($(Project)__GauntletTest)" />
-			
+		<Property Name="TestArg" Value="-test=$($(ProjectName)__GauntletTest)" />
 		<Property Name="RunArgs" Value="$(TestArg) $(TestFilterArgs) $(Args) -editorargs=&quot;$(CombinedEditorArgs)&quot;"/>
 
 		<!-- Do DDC cleanup outside the editor process, handles Cold cleanup without changing stats -->
@@ -270,6 +269,7 @@
 		<Property Name="TestName" Value="Editor.Workflows"/>
 		<Property Name="TestFilterArgs" Value="-runtest=$(TestName)"/>
 
+		<Property Name="EditorArgs" Value=""/>
 		<Property Name="TestOutputDir" Value="$(AutomationOutputDir)/$(ProjectName)$(Platform)EditorWorkflows"/>
 		<Property Name="TraceFilePath" Value="$(TraceDir)/$(TestName)/$(ProjectName)$(Platform)EditorWorkflows.utrace" />
 
@@ -277,7 +277,10 @@
 		<Property Name="ReportParams" Value="-WriteTestResultsForHorde -HordeTestDataKey=&quot;$(NodeName)$(PreflightSuffix)&quot;"/>
 		<Property Name="AllGauntletArgs" Value="-Project=$(ProjectName) -Platform=$(Platform) $(GauntletArgs) $(ReportParams) $(LogOutputParams)"/>
 
-		<Command Name="RunUnreal" Arguments="$(AllGauntletArgs) $(ProfileArgs) $(TestFilterArgs) -sessionlabel=&quot;$(TestName)&quot; -tracefile=&quot;$(TraceFilePath)&quot;" />
+		<Property Name="TestArg" Value="-test=$($(ProjectName)__GauntletTest)" />
+		<Property Name="RunArgs" Value="$(TestArg) $(TestFilterArgs) -editorargs=&quot;$(EditorArgs)&quot;"/>
+
+		<Command Name="RunUnreal" Arguments="$(AllGauntletArgs) $(ProfileArgs) $(RunArgs) -sessionlabel=&quot;$(TestName)&quot; -tracefile=&quot;$(TraceFilePath)&quot;" />
 	</Macro>
 	
 	<Macro Name="UploadTestTraces" Arguments="Platform;ProjectName">		
