Adds a new step to TSAN error process to allow a report of only the unique TSAN errors to be displayed. This step parses the output log from the test run and displays the unique errors and removes any interwoven log lines that are irrelevant.

Doing so the TSAN log matcher can now generate logs that contain the full TSAN error context in a single log. An TSAN IssueHandler has been added which looks for the summary log event and generates an issue keyed off the summary error reason and the file in error. The goal being to generate an issue per unique TSAN error

The TSANLogMatcher has had a test added to ensure we properly log only TSAN errors, and the file was renamed to fix warnings regarding file name and class name must matching.

EditorPerf timeouts have been reduced as we were waiting for 8 hours in cases when no tests were running due to errors in the config

diff --git a/Engine/Build/EditorPerf.xml b/Engine/Build/EditorPerf.xml
--- a/Engine/Build/EditorPerf.xml
+++ b/Engine/Build/EditorPerf.xml
@@ -76,6 +76,7 @@
 
 	<!-- Arguments for phases, keep in mind that these are passed to both the DDCCleanupCommandlet and RunUnreal -->
 	<Option Name="EnableBasePhaseTests" DefaultValue="true" Description="Enable the base phase perf tests."/>
+	<Property Name="EnableBasePhaseTests" Value="false" If="$(WithTSAN)"/>
 	<Option Name="ColdArgs" DefaultValue="-ddc=Cold -NoAssetRegistryCacheRead" Description="Arguments to unreal for the cold phase"/>
 	<Option Name="WarmArgs" DefaultValue="-ddc=Warm" Description="Arguments to unreal for the warm phase"/>
 
@@ -183,17 +184,14 @@
 
 	<!-- Gauntlet arguments for all runs -->
 	<Property Name="GauntletArgs" Value="-build=editor -tempdir=&quot;$(RootDir)/Tests&quot;"/>
-	<!-- TSAN prints to stdout/stderr. Using -VeryVerbose allows those streams to be redirected and captured by Gauntlet -->
-	<Property Name="GauntletArgs" Value="$(GauntletArgs) -VeryVerbose" If="'$(WithTSAN)'=='true'"/>
-	<Property Name="GauntletArgs" Value="$(GauntletArgs) -verbose" If="'$(WithTSAN)'=='false'"/>
 	<Property Name="GauntletArgs" Value="$(GauntletArgs) -branch=$(Branch) -ECBranch=&quot;$(Branch)&quot; -ECChangelist=&quot;$(Change)&quot;"/>
-	<Property Name="GauntletArgs" Value="$(GauntletArgs) -maxduration=28800"/>
+	<Property Name="GauntletArgs" Value="$(GauntletArgs) -verbose -maxduration=7200"/>
 
 	<!-- Gauntlet arguments for profile runs -->
 	<Property Name="ProfileArgs" Value="-tracefiletrunc"/>
 	<Property Name="ProfileArgs" Value="$(ProfileArgs) -EnablePlugins=$(ExtraPluginsCommaDelims)" If="'$(ExtraPluginsCommaDelims)' != ''"/>
 	<Property Name="ProfileArgs" Value="$(ProfileArgs) -SuppressLogWarnings -SuppressLogErrors"/>
-	<Property Name="ProfileArgs" Value="$(ProfileArgs) -GameInstanceLostTimerSeconds=3600"/>
+	<Property Name="ProfileArgs" Value="$(ProfileArgs) -GameInstanceLostTimerSeconds=600"/>
 	<Property Name="ProfileArgs" Value="$(ProfileArgs) -NoLoadLevelAtStartup -Benchmarking -NODDCCLEANUP"/>
 	<Property Name="ProfileArgs" Value="$(ProfileArgs) -NoShaderDistrib -VerboseShaderLogging"/>
 	<Property Name="ProfileArgs" Value="$(ProfileArgs) -ResumeOnCriticalFailure"/>
@@ -214,6 +212,8 @@
 	<Property Name="ExtraTelemetryArgs" Value="$(ExtraTelemetryArgs) -JobLink=&quot;$(HordeDashboardURL)/job/$(UE_HORDE_JOBID)?step=$(UE_HORDE_STEPID)&quot;" If="'$(HordeDashboardURL)' != ''"/>
 
 	<!-- The main test running macro logic that is expanded with platform and dependency node names -->
+	<Property Name="TestOutputDir" Value=""/>
+	<Property Name="RunTestLogPath" Value=""/>
 	<Macro Name="RunTest" Arguments="TestNameSuffix;Platform;ProjectName;ProjectTarget;NodeName;Phase;Iteration;Args;EditorArgs">
 		<Property Name="TestName" Value="EditorPerf$(Phase)$(TestNameSuffix)"/>
 		<Property Name="AssetRegistryCacheRootFolderArg" Value="-AssetRegistryCacheRootFolder=&quot;$(AssetRegistryCache_Folder)&quot;"/>
@@ -264,6 +264,17 @@
 		</Do>		
 	</Macro>
 	
+	<Property Name="GenerateTSANReportNodeName" Value=""/>
+	<Property Name="GenerateTSANReportNodeList" Value=""/>
+	<Macro Name="GenerateTSANReport" Arguments="NodeName;TestOutputDir;Phase;TestNameSuffix">
+		<Property Name="LogFilePath" Value="$(TestOutputDir)/$(Phase)$(TestNameSuffix)/Editor/EditorOutput.log"/>
+		<Property Name="GenerateTSANReportNodeName" Value="Generate $(NodeName) TSAN Report"/>
+		<Node Name="$(GenerateTSANReportNodeName)">
+			<Command Name="LogThreadSanitizerOutput" Arguments="-logfile=&quot;$(LogFilePath)&quot;" />
+		</Node>
+		<Property Name="GenerateTSANReportNodeList" Value="$(GenerateTSANReportNodeList); $(GenerateTSANReportNodeName)"/>
+	</Macro>
+	
 	<!-- The main editor workflows test running macro logic that is expanded with platform and dependency node names -->
 	<Macro Name="RunEditorWorkflowsTest" Arguments="Platform;ProjectName;NodeName">
 		<Property Name="TestName" Value="Editor.Workflows"/>
@@ -277,7 +288,7 @@
 		<Property Name="ReportParams" Value="-WriteTestResultsForHorde -HordeTestDataKey=&quot;$(NodeName)$(PreflightSuffix)&quot;"/>
 		<Property Name="AllGauntletArgs" Value="-Project=$(ProjectName) -Platform=$(Platform) $(GauntletArgs) $(ReportParams) $(LogOutputParams)"/>
 
-		<Property Name="TestArg" Value="-test=$($(ProjectName)__GauntletTest)" />
+		<Property Name="TestArg" Value="-test=EditorTests" />
 		<Property Name="RunArgs" Value="$(TestArg) $(TestFilterArgs) -editorargs=&quot;$(EditorArgs)&quot;"/>
 
 		<Command Name="RunUnreal" Arguments="$(AllGauntletArgs) $(ProfileArgs) $(RunArgs) -sessionlabel=&quot;$(TestName)&quot; -tracefile=&quot;$(TraceFilePath)&quot;" />
@@ -361,6 +372,7 @@
 				<Property Name="RunNodeList" Value=""/>
 				<!--Each node in this the nested loops below, they need to depend on predecessors--> 
 				<Property Name="UploadNodeList" Value=""/>
+				
 				<Agent Name="EditorPerf Agent $(Platform) $(Project)" Type="EditorPerf$(Platform)">
 					<ForEach Name="Phase" Values="$(Phases)">
 						<Do If="$(EnableBasePhaseTests)">
@@ -372,8 +384,10 @@
 									<Expand Name="RunTest" TestNameSuffix="NoZenLoader" Platform="$(Platform)" ProjectName="$(Project)" ProjectTarget="$(ProjectTarget)" NodeName="$(RunNodeName)" Phase="$(Phase)" Iteration="$(Iteration)" Args="$($(Phase)Args)" EditorArgs="$(ProjectCombinedArgs) $(NoZenLoaderTestArgs)"/>
 								</Node>
 								<Property Name="RunNodeList" Value="$(RunNodeList);$(RunNodeName)"/>
+								<Expand Name="GenerateTSANReport" NodeName="$(RunNodeName)" TestOutputDir="$(TestOutputDir)" Phase="$(Phase)" TestNameSuffix="NoZenLoader" If="$(WithTSAN)"/>
 							</ForEach>
 						</Do>
+						
 						<!-- Setup Zen Loader tests for this phase -->
 						<Do>
 							<!--Zen Loader TestRun if enabled--> 
@@ -382,12 +396,15 @@
 								<Expand Name="RunTest" TestNameSuffix="ZenLoader" Platform="$(Platform)" ProjectName="$(Project)" ProjectTarget="$(ProjectTarget)" NodeName="$(RunNodeZenName)" Phase="$(Phase)" Iteration="1" Args="$($(Phase)Args)" EditorArgs="$(ProjectCombinedArgs) $(ZenLoaderTestArgs)"/>
 							</Node>
 							<Property Name="RunNodeList" Value="$(RunNodeList);$(RunNodeZenName)" If="$(EnableZenLoaderTest)"/>
+							<Expand Name="GenerateTSANReport" NodeName="$(RunNodeZenName)" TestOutputDir="$(TestOutputDir)" Phase="$(Phase)" TestNameSuffix="ZenLoader" If="$(EnableZenLoaderTest) And $(WithTSAN)"/>
+							
 							<!--Zen Loader Async TestRun if enabled--> 
 							<Property Name="RunNodeZenAsyncName" Value="Run $(OutputName) $(Phase) ZenLoader Async"/>
 							<Node Name="$(RunNodeZenAsyncName)" Requires="$(RunNodeRequires);$(RunNodeList)" If="$(EnableZenLoaderAsyncThreadTest)">
 								<Expand Name="RunTest" TestNameSuffix="ZenLoaderAsync" Platform="$(Platform)" ProjectName="$(Project)" ProjectTarget="$(ProjectTarget)" NodeName="$(RunNodeZenAsyncName)" Phase="$(Phase)" Iteration="1" Args="$($(Phase)Args)" EditorArgs="$(ProjectCombinedArgs) $(ZenLoaderAsyncTestArgs)"/>
 							</Node>
 							<Property Name="RunNodeList" Value="$(RunNodeList);$(RunNodeZenAsyncName)" If="$(EnableZenLoaderAsyncThreadTest)"/>
+							<Expand Name="GenerateTSANReport" NodeName="$(RunNodeZenAsyncName)" TestOutputDir="$(TestOutputDir)" Phase="$(Phase)" TestNameSuffix="ZenLoaderAsync" If="$(EnableZenLoaderAsyncThreadTest) And $(WithTSAN)"/>
 
 							<!--Zen Store DDC TestRun if enabled--> 
 							<Property Name="RunNodeZenStoreDDCName" Value="Run $(OutputName) $(Phase) ZenLoader Store DDC"/>
@@ -395,7 +412,9 @@
 								<Expand Name="RunTest" TestNameSuffix="ZenStoreAsyncDDC" Platform="$(Platform)" ProjectName="$(Project)" ProjectTarget="$(ProjectTarget)" NodeName="$(RunNodeZenStoreDDCName)" Phase="$(Phase)" Iteration="1" Args="$(ZenStoreDDCPhaseArgs)" EditorArgs="$(ProjectCombinedArgs) $(ZenStoreDDCTestArgs)"/>
 							</Node>
 							<Property Name="RunNodeList" Value="$(RunNodeList);$(RunNodeZenStoreDDCName)" If="$(EnableZenStoreDDCTest)"/>
+							<Expand Name="GenerateTSANReport" NodeName="$(RunNodeZenStoreDDCName)" TestOutputDir="$(TestOutputDir)" Phase="$(Phase)" TestNameSuffix="ZenStoreAsyncDDC" If="$(EnableZenStoreDDCTest) And $(WithTSAN)"/>
 						</Do>
+						
 						<!--Warm only test runs e.g LazyLoad--> 
 						<Do If="'$(Phase)' == 'Warm'">
 							<!--Legacy LazyLoad TestRun if enabled--> 
@@ -404,14 +423,18 @@
 								<Expand Name="RunTest" TestNameSuffix="NoZenLoader" Platform="$(Platform)" ProjectName="$(Project)" ProjectTarget="$(ProjectTarget)" NodeName="$(RunNodeLazyLoadName)" Phase="$(Phase)" Iteration="1" Args="$($(Phase)Args)" EditorArgs="$(ProjectCombinedArgs) $(LazyLoadTestArgs) $(NoZenLoaderTestArgs)"/>
 							</Node>
 							<Property Name="RunNodeList" Value="$(RunNodeList);$(RunNodeLazyLoadName)" If="$(EnableLazyLoadNoZenLoaderTest)"/>
+							<Expand Name="GenerateTSANReport" NodeName="$(RunNodeLazyLoadName)" TestOutputDir="$(TestOutputDir)" Phase="$(Phase)" TestNameSuffix="NoZenLoader" If="$(EnableLazyLoadNoZenLoaderTest) And $(WithTSAN)"/>
+							
 							<!--Zen LazyLoad TestRun if enabled--> 
 							<Property Name="RunNodeLazyLoadZenName" Value="Run $(OutputName) $(Phase) LazyLoad ZenLoader"/>
 							<Node Name="$(RunNodeLazyLoadZenName)" Requires="$(RunNodeRequires);$(RunNodeList)" If="$(EnableLazyLoadZenLoaderTest)">
 								<Expand Name="RunTest" TestNameSuffix="ZenLoader" Platform="$(Platform)" ProjectName="$(Project)" ProjectTarget="$(ProjectTarget)" NodeName="$(RunNodeLazyLoadZenName)" Phase="$(Phase)" Iteration="1" Args="$($(Phase)Args)" EditorArgs="$(ProjectCombinedArgs) $(LazyLoadZenTestArgs)"/>
 							</Node>
 							<Property Name="RunNodeList" Value="$(RunNodeList);$(RunNodeLazyLoadZenName)" If="$(EnableLazyLoadZenLoaderTest)"/>
+							<Expand Name="GenerateTSANReport" NodeName="$(RunNodeLazyLoadZenName)" TestOutputDir="$(TestOutputDir)" Phase="$(Phase)" TestNameSuffix="ZenLoader" If="$(EnableLazyLoadZenLoaderTest) And $(WithTSAN)"/>
 						</Do>
 					</ForEach>
+					
 					<Property Name="RunNodeUploadPhaseTraces" Value="Upload $(OutputName) Traces"/>
 					<Node Name="$(RunNodeUploadPhaseTraces)">
 						<Expand Name="UploadTestTraces" Platform="$(Platform)" ProjectName="$(Project)"/>
@@ -419,7 +442,7 @@
 					<Property Name="UploadNodeList" Value="$(RunNodeUploadPhaseTraces)"/>
 				</Agent>
 				<Label Category="$(RunCategory)" Name="$(Project) Editor Run" Requires="$(RunNodeRequires);$(RunNodeList);$(UploadNodeList)" UgsBadge="$(Project) Editor Run"/>
-				<Property Name="RunGraphNodes" Value="$(RunGraphNodes);$(RunNodeList);$(UploadNodeList)"/>
+				<Property Name="RunGraphNodes" Value="$(RunGraphNodes);$(RunNodeList);$(GenerateTSANReportNodeList);$(UploadNodeList)"/>
 
 				<!--Baseline the Results--> 
 				<Property Name="BaselineNodeName" Value="Baseline $(OutputName)"/>
diff --git a/Engine/Source/Programs/AutomationTool.Tests/LogParserTests.cs b/Engine/Source/Programs/AutomationTool.Tests/LogParserTests.cs
--- a/Engine/Source/Programs/AutomationTool.Tests/LogParserTests.cs
+++ b/Engine/Source/Programs/AutomationTool.Tests/LogParserTests.cs
@@ -596,7 +596,6 @@ namespace AutomationTool.Tests
 			CheckEventGroup(logEvents, 2, 1, LogLevel.Warning, KnownLogEvents.Compiler);
 		}
 
-		[Ignore]
 		[TestMethod]
 		public void DockerErrorMatcher()
 		{
@@ -700,6 +699,50 @@ namespace AutomationTool.Tests
 			}
 		}
 
+		[TestMethod]
+		public void ThreadSanitizerErrorMatcher()
+		{
+			string[] lines =
+			{
+				@"==================",
+				@"WARNING: ThreadSanitizer: data race (pid=21089)",
+				@"  Write of size 1 at 0x00004060ea38 by thread T25:",
+				@"    #0 FPaths::IsStaged() /mnt/horde/++UE5/Sync/Engine/Source/./Runtime/Core/Private/Misc/Paths.cpp:172:33 (CitySampleEditor+0x2c0f2f44) (BuildId: 6622c84dbb6a946e)",
+				@"",
+				@"  Previous write of size 1 at 0x00004060ea38 by thread T20:",
+				@"    #0 FPaths::IsStaged() /mnt/horde/++UE5/Sync/Engine/Source/./Runtime/Core/Private/Misc/Paths.cpp:172:33 (CitySampleEditor+0x2c0f2f44) (BuildId: 6622c84dbb6a946e)",
+				@"",
+				@"  Location is global '??' at 0x000000000000 (CitySampleEditor+0x4060ea38)",
+				@"",
+				@"  Thread T25 'Backgro-ker #19' (tid=21138, running) created by thread T7 at:",
+				@"    #0 pthread_create /src/build/llvm-src/compiler-rt/lib/tsan/rtl/tsan_interceptors_posix.cpp:1048 (CitySampleEditor+0x10351c92) (BuildId: 6622c84dbb6a946e)",
+				@"    #1 FRunnableThreadPThread::CreateThreadWithName(unsigned long*, pthread_attr_t*, void* (*)(void*), void*, char const*) /mnt/horde/++UE5/Sync/Engine/Source/Runtime/Core/Private/HAL/PThreadRunnableThread.h:90:10 (CitySampleEditor+0x2bd32c50) (BuildId: 6622c84dbb6a946e)",
+				@"",
+				@"  Thread T20 'Backgro-ker #14' (tid=21133, running) created by thread T5 at:",
+				@"    #0 pthread_create /src/build/llvm-src/compiler-rt/lib/tsan/rtl/tsan_interceptors_posix.cpp:1048 (CitySampleEditor+0x10351c92) (BuildId: 6622c84dbb6a946e)",
+				@"    #1 FRunnableThreadPThread::CreateThreadWithName(unsigned long*, pthread_attr_t*, void* (*)(void*), void*, char const*) /mnt/horde/++UE5/Sync/Engine/Source/Runtime/Core/Private/HAL/PThreadRunnableThread.h:90:10 (CitySampleEditor+0x2bd32c50) (BuildId: 6622c84dbb6a946e)",
+				@"",
+				@"SUMMARY: ThreadSanitizer: data race /mnt/horde/++UE5/Sync/Engine/Source/./Runtime/Core/Private/Misc/Paths.cpp:172:33 in FPaths::IsStaged()",
+				@"=================="
+			};
+
+			List<LogEvent> logEvents = Parse(lines);
+			Assert.AreEqual(20, logEvents[0].LineCount);
+			CheckEventGroup(logEvents, 0, 20, LogLevel.Error, KnownLogEvents.Sanitizer_Thread);
+
+			Assert.AreEqual("/mnt/horde/++UE5/Sync/Engine/Source/Runtime/Core/Private/HAL/PThreadRunnableThread.h", logEvents[12].GetProperty("SourceFile").ToString());
+			Assert.AreEqual("90", logEvents[12].GetProperty("Line").ToString());
+			Assert.AreEqual("10", logEvents[12].GetProperty("Column").ToString());
+			Assert.AreEqual("FRunnableThreadPThread::CreateThreadWithName(unsigned long*, pthread_attr_t*, void* (*)(void*), void*, char const*)", logEvents[12].GetProperty("Symbol").ToString());
+
+			// Summary
+			Assert.AreEqual("/mnt/horde/++UE5/Sync/Engine/Source/./Runtime/Core/Private/Misc/Paths.cpp", logEvents[18].GetProperty("SummarySourceFile").ToString());
+			Assert.AreEqual("172", logEvents[18].GetProperty("Line").ToString());
+			Assert.AreEqual("33", logEvents[18].GetProperty("Column").ToString());
+			Assert.AreEqual("FPaths::IsStaged()", logEvents[18].GetProperty("Symbol").ToString());
+			Assert.AreEqual("data race", logEvents[18].GetProperty("SummaryReason").ToString());
+		}
+
 		[TestMethod]
 		public void LocalizationChannelMatcher()
 		{
diff --git a/Engine/Source/Programs/AutomationTool/AutomationUtils/Matchers/ThreadSanitizerEventMatcher.cs b/Engine/Source/Programs/AutomationTool/AutomationUtils/Matchers/ThreadSanitizerEventMatcher.cs
new file mode 100644
--- /dev/null
+++ b/Engine/Source/Programs/AutomationTool/AutomationUtils/Matchers/ThreadSanitizerEventMatcher.cs
@@ -0,0 +1,87 @@
+// Copyright Epic Games, Inc. All Rights Reserved.
+
+using EpicGames.Core;
+using EpicGames.Horde.Issues.Handlers;
+using Microsoft.Extensions.Logging;
+using System.Text.RegularExpressions;
+
+#nullable enable
+
+namespace AutomationUtils.Matchers
+{
+	/// <summary>
+	/// Matcher for generated Thread Sanitizer reports
+	/// 
+	/// </summary>
+	class ThreadSanitizerEventMatcher : ILogEventMatcher
+	{
+		// All TSAN reports start and end with this delimiter.
+		static readonly Regex s_StartEndDelimiter = new Regex(@"^==================$");
+
+		// e.g. Note line number, column number and BuildId are optional
+		// #1 MyType<Type, OtherType>::Function(unsigned long*&) const /mnt/somepath/file.h:90:10 (BinaryName+0x2bc102e0) (BuildId: 8b17597a9f5444a0)
+		static readonly Regex s_StackTracePattern = new Regex(
+			@"^\s+\#[\d]+?\s" +
+			@"(?<Symbol>.+?)\s" +
+			@"(?<SourceFile>(\/|\w:).+?)" +
+			@"(:(?<Line>\d+?))?" +
+			@"(:(?<Column>\d+?))?" +
+			@"\s+\(.+?\+(?<Address>0x[0-9a-fA-F]+?)\)" +
+			@"(\s+?\(BuildId:\s+?(?<BuildId>.+?)\))?$");
+
+		// e.g.
+		// SUMMARY: ThreadSanitizer: data race /some/path.cpp:169:33 in MyType::Foo()
+		static readonly Regex s_SummaryPattern = new Regex(
+			@"^SUMMARY: ThreadSanitizer:\s*" +
+			@"(?<SummaryReason>.+)\s" +
+			@"(?<SourceFile>(\/|\w:).+?)" +
+			@"(:(?<Line>\d+?))?" +
+			@"(:(?<Column>\d+?))?\s+?in\s+?" + 
+			@"(?<Symbol>.+)$");
+
+		public LogEventMatch? Match(ILogCursor cursor)
+		{
+			if (cursor.IsMatch(s_StartEndDelimiter))
+			{
+				Match? match;
+				LogEventBuilder builder = new LogEventBuilder(cursor);
+				builder.MoveNext();
+
+				while (builder.Current.CurrentLine != null && !builder.Current.IsMatch(s_StartEndDelimiter))
+				{
+					if (builder.Current.TryMatch(s_StackTracePattern, out match))
+					{
+						do
+						{
+							builder.AnnotateSymbol(match!.Groups["Symbol"]);
+							builder.AnnotateSourceFile(match.Groups["SourceFile"], "");
+							builder.Annotate(match.Groups["Line"], LogEventMarkup.LineNumber);
+							builder.Annotate(match.Groups["Column"], LogEventMarkup.ColumnNumber);
+							builder.Annotate(match.Groups["Address"], "");
+							builder.Annotate(match.Groups["BuildId"], "");
+
+							builder.MoveNext();
+						} while (builder.Current.TryMatch(s_StackTracePattern, out match));
+					}
+
+					if (builder.Current.TryMatch(s_SummaryPattern, out match))
+					{
+						// Used by IssueHandler
+						builder.Annotate(match.Groups["SummaryReason"], ThreadSanitizerIssueHandler.SummaryReason);
+						builder.Annotate("SummarySourceFile", match.Groups["SourceFile"], ThreadSanitizerIssueHandler.SummarySourceFile);
+
+						// Annotate the source file again so that we may get UGS link resolution
+						builder.AnnotateSourceFile(match.Groups["SourceFile"], "");
+						builder.Annotate(match.Groups["Line"], LogEventMarkup.LineNumber);
+						builder.Annotate(match.Groups["Column"], LogEventMarkup.ColumnNumber);
+						builder.AnnotateSymbol(match.Groups["Symbol"]);
+					}
+					builder.MoveNext();
+				}
+				return builder.ToMatch(LogEventPriority.Normal, LogLevel.Error, KnownLogEvents.Sanitizer_Thread);
+			}
+
+			return null;
+		}
+	}
+}
diff --git a/Engine/Source/Programs/AutomationTool/AutomationUtils/Matchers/ThreadSanitizerMatcher.cs b/Engine/Source/Programs/AutomationTool/AutomationUtils/Matchers/ThreadSanitizerMatcher.cs
deleted file mode 100644
diff --git a/Engine/Source/Programs/Horde/Horde.Server.Tests/Issues/IssueServiceTests.cs b/Engine/Source/Programs/Horde/Horde.Server.Tests/Issues/IssueServiceTests.cs
--- a/Engine/Source/Programs/Horde/Horde.Server.Tests/Issues/IssueServiceTests.cs
+++ b/Engine/Source/Programs/Horde/Horde.Server.Tests/Issues/IssueServiceTests.cs
@@ -662,6 +662,115 @@ namespace Horde.Server.Tests.Issues
 			}
 		}
 
+		[TestMethod]
+		public async Task ThreadSanitizerDeduplicateSimilarIssuesTestAsync()
+		{
+			// We have three unique errors reported, but we key issues on the file since file lines change over time without fixes, so only two unique issues should be found
+			string[] lines =
+			{
+				@"==================",
+				@"WARNING: ThreadSanitizer: data race (pid=45069)",
+				@"  Write of size 1 at 0x0000408ba588 by thread T15:",
+				@"    #0 FPaths::IsStaged() /mnt/horde/++UE5/Sync/Engine/Source/./Runtime/Core/Private/Misc/Paths.cpp:164:33 (CitySampleEditor+0x2c2561a3) (BuildId: 7ca67a0f97321f67)",
+				@"    #1 FGenericPlatformMisc::ProjectDir() /mnt/horde/++UE5/Sync/Engine/Source/./Runtime/Core/Private/GenericPlatform/GenericPlatformMisc.cpp:1264:19 (CitySampleEditor+0x2be6a866) (BuildId: 7ca67a0f97321f67)",
+				@"",
+				@"  Previous write of size 1 at 0x0000408ba588 by thread T13:",
+				@"    #0 FPaths::IsStaged() /mnt/horde/++UE5/Sync/Engine/Source/./Runtime/Core/Private/Misc/Paths.cpp:164:33 (CitySampleEditor+0x2c2561a3) (BuildId: 7ca67a0f97321f67)",
+				@"",
+				@"SUMMARY: ThreadSanitizer: data race /mnt/horde/++UE5/Sync/Engine/Source/./Runtime/Core/Private/Misc/Paths.cpp:164:33 in FPaths::IsStaged()",
+				@"==================",
+				@"",
+				@"==================",
+				@"WARNING: ThreadSanitizer: data race (pid=45069)",
+				@"  Write of size 1 at 0x0000408ba588 by thread T23:",
+				@"    #0 FPaths::IsStaged() /mnt/horde/++UE5/Sync/Engine/Source/./Runtime/Core/Private/Misc/Paths.cpp:172:33 (CitySampleEditor+0x2c256324) (BuildId: 7ca67a0f97321f67)",
+				@"",
+				@"  Thread T23 'Backgro-ker #17' (tid=45131, running) created by thread T15 at:",
+				@"    #0 pthread_create /src/build/llvm-src/compiler-rt/lib/tsan/rtl/tsan_interceptors_posix.cpp:1048 (CitySampleEditor+0x103ffc92) (BuildId: 7ca67a0f97321f67)",
+				@"",
+				@"SUMMARY: ThreadSanitizer: data race /mnt/horde/++UE5/Sync/Engine/Source/./Runtime/Core/Private/Misc/Paths.cpp:172:33 in FPaths::IsStaged()",
+				@"==================",
+				@"",
+				@"==================",
+				@"WARNING: ThreadSanitizer: data race (pid=45069)",
+				@"  Read of size 1 at 0x0000408b68d8 by thread T47 (mutexes: write M0, write M1):",
+				@"    #0 IsEngineStartupModuleLoadingComplete() /mnt/horde/++UE5/Sync/Engine/Source/./Runtime/Core/Private/Misc/CoreGlobals.cpp:297:9 (CitySampleEditor+0x2c1bb5c8) (BuildId: 7ca67a0f97321f67)",
+				@"",
+				@"SUMMARY: ThreadSanitizer: data race /mnt/horde/++UE5/Sync/Engine/Source/./Runtime/Core/Private/Misc/CoreGlobals.cpp:297:9 in IsEngineStartupModuleLoadingComplete()",
+				@"=================="
+			};
+
+			IJob job = CreateJob(_mainStreamId, 120, "Generate SomeJobs TSAN Report", _graph);
+			await ParseEventsAsync(job, 0, 0, lines);
+			await UpdateCompleteStepAsync(job, 0, 0, JobStepOutcome.Failure);
+
+			IReadOnlyList<IIssue> issues = await IssueCollection.FindIssuesAsync();
+			Assert.AreEqual(2, issues.Count);
+
+			IIssue issue1 = issues[0];
+			Assert.AreEqual(issue1.Fingerprints.Count, 1);
+			Assert.IsTrue(issue1.Fingerprints[0].Type.Contains("ThreadSanitizer", StringComparison.OrdinalIgnoreCase));
+			Assert.IsTrue(issue1.Fingerprints[0].Type.Contains("data race", StringComparison.OrdinalIgnoreCase));
+			Assert.IsTrue(issue1.Fingerprints[0].Type.Contains("CoreGlobals.cpp", StringComparison.OrdinalIgnoreCase));
+
+			IIssue issue2 = issues[1];
+			Assert.AreEqual(issue2.Fingerprints.Count, 1);
+			Assert.IsTrue(issue2.Fingerprints[0].Type.Contains("ThreadSanitizer", StringComparison.OrdinalIgnoreCase));
+			Assert.IsTrue(issue2.Fingerprints[0].Type.Contains("data race", StringComparison.OrdinalIgnoreCase));
+			Assert.IsTrue(issue2.Fingerprints[0].Type.Contains("Paths.cpp", StringComparison.OrdinalIgnoreCase));
+		}
+
+		[TestMethod]
+		public async Task ThreadSanitizerEnsureUniqueIssuesForSameFileTestAsync()
+		{
+			// Generate two errors for the same file and line number but with a different summary reason. We should get two issues as the problems are unique
+			string[] lines =
+			{
+				@"==================",
+				@"WARNING: ThreadSanitizer: data race (pid=45069)",
+				@"  Write of size 1 at 0x0000408ba588 by thread T15:",
+				@"    #0 FPaths::IsStaged() /mnt/horde/++UE5/Sync/Engine/Source/./Runtime/Core/Private/Misc/Paths.cpp:164:33 (CitySampleEditor+0x2c2561a3) (BuildId: 7ca67a0f97321f67)",
+				@"    #1 FGenericPlatformMisc::ProjectDir() /mnt/horde/++UE5/Sync/Engine/Source/./Runtime/Core/Private/GenericPlatform/GenericPlatformMisc.cpp:1264:19 (CitySampleEditor+0x2be6a866) (BuildId: 7ca67a0f97321f67)",
+				@"",
+				@"  Previous write of size 1 at 0x0000408ba588 by thread T13:",
+				@"    #0 FPaths::IsStaged() /mnt/horde/++UE5/Sync/Engine/Source/./Runtime/Core/Private/Misc/Paths.cpp:164:33 (CitySampleEditor+0x2c2561a3) (BuildId: 7ca67a0f97321f67)",
+				@"",
+				@"SUMMARY: ThreadSanitizer: data race /mnt/horde/++UE5/Sync/Engine/Source/./Runtime/Core/Private/Misc/Paths.cpp:164:33 in FPaths::IsStaged()",
+				@"==================",
+				@"",
+				@"==================",
+				@"WARNING: ThreadSanitizer: lock-order-inversion (potential deadlock) (pid=45069)",
+				@"  Write of size 1 at 0x0000408ba588 by thread T15:",
+				@"    #0 FPaths::IsStaged() /mnt/horde/++UE5/Sync/Engine/Source/./Runtime/Core/Private/Misc/Paths.cpp:164:33 (CitySampleEditor+0x2c2561a3) (BuildId: 7ca67a0f97321f67)",
+				@"    #1 FGenericPlatformMisc::ProjectDir() /mnt/horde/++UE5/Sync/Engine/Source/./Runtime/Core/Private/GenericPlatform/GenericPlatformMisc.cpp:1264:19 (CitySampleEditor+0x2be6a866) (BuildId: 7ca67a0f97321f67)",
+				@"",
+				@"  Previous write of size 1 at 0x0000408ba588 by thread T13:",
+				@"    #0 FPaths::IsStaged() /mnt/horde/++UE5/Sync/Engine/Source/./Runtime/Core/Private/Misc/Paths.cpp:164:33 (CitySampleEditor+0x2c2561a3) (BuildId: 7ca67a0f97321f67)",
+				@"",
+				@"SUMMARY: ThreadSanitizer: lock-order-inversion (potential deadlock) /mnt/horde/++UE5/Sync/Engine/Source/./Runtime/Core/Private/Misc/Paths.cpp:172:33 in FPaths::IsStaged()",
+				@"==================",
+			};
+
+			IJob job = CreateJob(_mainStreamId, 120, "Generate SomeJobs TSAN Report", _graph);
+			await ParseEventsAsync(job, 0, 0, lines);
+			await UpdateCompleteStepAsync(job, 0, 0, JobStepOutcome.Failure);
+
+			IReadOnlyList<IIssue> issues = await IssueCollection.FindIssuesAsync();
+			Assert.AreEqual(2, issues.Count);
+
+			IIssue issue1 = issues[0];
+			Assert.AreEqual(issue1.Fingerprints.Count, 1);
+			Assert.IsTrue(issue1.Fingerprints[0].Type.Contains("ThreadSanitizer", StringComparison.OrdinalIgnoreCase));
+			Assert.IsTrue(issue1.Fingerprints[0].Type.Contains("lock-order-inversion", StringComparison.OrdinalIgnoreCase));
+			Assert.IsTrue(issue1.Fingerprints[0].Type.Contains("Paths.cpp", StringComparison.OrdinalIgnoreCase));
+
+			IIssue issue2 = issues[1];
+			Assert.AreEqual(issue2.Fingerprints.Count, 1);
+			Assert.IsTrue(issue2.Fingerprints[0].Type.Contains("ThreadSanitizer", StringComparison.OrdinalIgnoreCase));
+			Assert.IsTrue(issue2.Fingerprints[0].Type.Contains("data race", StringComparison.OrdinalIgnoreCase));
+			Assert.IsTrue(issue2.Fingerprints[0].Type.Contains("Paths.cpp", StringComparison.OrdinalIgnoreCase)); 
+		}
+
 		[TestMethod]
 		public async Task AutoSdkWarningTestAsync()
 		{
diff --git a/Engine/Source/Programs/Shared/EpicGames.Horde/Issues/Handlers/ThreadSanitizerIssueHandler.cs b/Engine/Source/Programs/Shared/EpicGames.Horde/Issues/Handlers/ThreadSanitizerIssueHandler.cs
new file mode 100644
--- /dev/null
+++ b/Engine/Source/Programs/Shared/EpicGames.Horde/Issues/Handlers/ThreadSanitizerIssueHandler.cs
@@ -0,0 +1,74 @@
+// Copyright Epic Games, Inc. All Rights Reserved.
+
+using System;
+using System.Collections.Generic;
+using System.IO;
+using System.Linq;
+using System.Text.Json;
+using EpicGames.Core;
+
+namespace EpicGames.Horde.Issues.Handlers
+{
+	/// <summary>
+	/// Instance of a specific Thread Sanitizer error
+	/// </summary>
+	[IssueHandler(Priority = 10)]
+	public class ThreadSanitizerIssueHandler : IssueHandler
+	{
+		static readonly Utf8String s_summaryReasonAnnotation = new Utf8String("SummaryReason");
+		static readonly Utf8String s_summarySourceFileAnnotation = new Utf8String("SummarySourceFile");
+		readonly List<IssueEventGroup> _issues = new List<IssueEventGroup>();
+
+		/// <summary>
+		/// Log value describing the thread sanitizer error summary reason
+		/// </summary>
+		public static LogValue SummaryReason => new LogValue(s_summaryReasonAnnotation, "");
+
+		/// <summary>
+		/// Log value describing the thread sanitizer error summary source file
+		/// </summary>
+		public static LogValue SummarySourceFile => new LogValue(s_summarySourceFileAnnotation, "");
+
+		static bool FindAnnotation(IssueEvent issueEvent, Utf8String annotation, ref string result)
+		{
+			foreach (JsonProperty property in issueEvent.Lines.SelectMany(x => x.FindPropertiesOfType(annotation)))
+			{
+				JsonElement value;
+				if (property.Value.TryGetProperty(LogEventPropertyName.Text.Span, out value) && value.ValueKind == JsonValueKind.String)
+				{
+					result = value.GetString() ?? "";
+					return true;
+				}
+			}
+			return false;
+		}
+
+		/// <inheritdoc/>
+		public override bool HandleEvent(IssueEvent issueEvent)
+		{
+			if (issueEvent.EventId == KnownLogEvents.Sanitizer_Thread)
+			{
+				string summaryReason = "error";
+				string summaryFile = "";
+
+				FindAnnotation(issueEvent, s_summaryReasonAnnotation, ref summaryReason);
+				FindAnnotation(issueEvent, s_summarySourceFileAnnotation, ref summaryFile);
+				string fingerprint = $"ThreadSanitizer:{summaryReason}:{summaryFile}";
+
+				IssueEventGroup issue = new IssueEventGroup(fingerprint, $"Thread sanitizer found '{{Meta:{s_summaryReasonAnnotation}}}' in {{Meta:{s_summarySourceFileAnnotation}}}", IssueChangeFilter.Code);
+				issue.Keys.Add(summaryReason, IssueKeyType.Note);
+				issue.Metadata.Add(s_summaryReasonAnnotation.ToString(), summaryReason);
+				issue.Keys.AddSourceFile(summaryFile, IssueKeyType.File);
+				issue.Metadata.Add(s_summarySourceFileAnnotation.ToString(), summaryFile);
+				issue.Events.Add(issueEvent);
+
+				_issues.Add(issue);
+				return true;
+			}
+			return false;
+		}
+
+		/// <inheritdoc/>
+		public override IEnumerable<IssueEventGroup> GetIssues() => _issues;
+	}
+}
