Merging

Main/Engine/Build/InstalledEngineFilters.xml
Main/Engine/Source/Developer/LowLevelTestsRunner/LowLevelTestsRunner.Build.cs
Main/Engine/Source/ThirdParty/Catch2/Catch2Extras.Build.cs

to //UE5/Release-5.4/Engine/Source/

diff --git a/Engine/Build/InstalledEngineFilters.xml b/Engine/Build/InstalledEngineFilters.xml
--- a/Engine/Build/InstalledEngineFilters.xml
+++ b/Engine/Build/InstalledEngineFilters.xml
@@ -326,6 +326,13 @@
 
 		<!-- nanosvg -->
 		Engine/Source/ThirdParty/nanosvg/...
+		
+		<!-- Catch2 -->
+		Engine/Source/ThirdParty/Catch2/v3.4.0/...
+		Engine/Source/ThirdParty/Catch2/Catch2_v3.4.0.tps
+		Engine/Source/ThirdParty/Catch2/Catch2.build.cs
+		Engine/Source/ThirdParty/Catch2/Catch2.tps
+		Engine/Source/ThirdParty/Catch2/Catch2Extras.Build.cs
 	</Property>
 
 	<Property Name="CopyEditorExceptionsWin64" Value="$(CopyEditorExceptions)">
@@ -383,6 +390,13 @@
 		Engine/Binaries/Mac/zenserver
 		Engine/Binaries/Mac/zen
 		Engine/Binaries/Mac/crashpad_handler
+			
+		<!-- Catch2 -->
+		Engine/Source/ThirdParty/Catch2/v3.4.0/...
+		Engine/Source/ThirdParty/Catch2/Catch2_v3.4.0.tps
+		Engine/Source/ThirdParty/Catch2/Catch2.build.cs
+		Engine/Source/ThirdParty/Catch2/Catch2.tps
+		Engine/Source/ThirdParty/Catch2/Catch2Extras.Build.cs
 	</Property>
 
 	<Property Name="CopyEditorExceptionsMac" Value="$(CopyEditorExceptions)">
@@ -444,6 +458,13 @@
 		<!-- zenserver is the tool for Unreal Zen Storage services like Local DDC and optionally cooked output storage -->
 		Engine/Binaries/Linux/zenserver
 		Engine/Binaries/Linux/zen
+			
+		<!-- Catch2 -->
+		Engine/Source/ThirdParty/Catch2/v3.4.0/...
+		Engine/Source/ThirdParty/Catch2/Catch2_v3.4.0.tps
+		Engine/Source/ThirdParty/Catch2/Catch2.build.cs
+		Engine/Source/ThirdParty/Catch2/Catch2.tps
+		Engine/Source/ThirdParty/Catch2/Catch2Extras.Build.cs
 	</Property>
 	<Property Name="CopyEditorExceptionsLinux" Value="$(CopyEditorExceptions)">
 		<!-- Exclude Windows/Mac binaries on Linux here -->
diff --git a/Engine/Source/Developer/LowLevelTestsRunner/LowLevelTestsRunner.Build.cs b/Engine/Source/Developer/LowLevelTestsRunner/LowLevelTestsRunner.Build.cs
--- a/Engine/Source/Developer/LowLevelTestsRunner/LowLevelTestsRunner.Build.cs
+++ b/Engine/Source/Developer/LowLevelTestsRunner/LowLevelTestsRunner.Build.cs
@@ -42,12 +42,6 @@ namespace UnrealBuildTool.Rules
 			{
 				PrivateDependencyModuleNames.Add("CoreUObject");
 			}
-
-			if (Target.Platform == UnrealTargetPlatform.Android)
-			{
-				string ModulePath = Utils.MakePathRelativeTo(ModuleDirectory, Target.RelativeEnginePath);
-				AdditionalPropertiesForReceipt.Add("AndroidPlugin", Path.Combine(ModulePath, "LowLevelTestsAndroid_UPL.xml"));
-			}
 		}
 	}
 }
