- Increases the timeout for TSAN EditorPerf tests to 5 hours (2 hours is timing out currently)
- Enables async loading perf tests config now that seems stable enough to get perf numbers and TSAN validation

#rnx

diff --git a/Engine/Build/EditorPerf.xml b/Engine/Build/EditorPerf.xml
--- a/Engine/Build/EditorPerf.xml
+++ b/Engine/Build/EditorPerf.xml
@@ -154,7 +154,7 @@
 
 	<!-- Options for extra Zen Loader tests -->
 	<Option Name="EnableZenLoaderTest" DefaultValue="true" Description="Enable additional test per phase that uses the Zen Loader."/>
-	<Option Name="EnableZenLoaderAsyncThreadTest" DefaultValue="false" Description="Enable additional test per phase that uses the Zen Loader with an async loading thread."/>
+	<Option Name="EnableZenLoaderAsyncThreadTest" DefaultValue="true" Description="Enable additional test per phase that uses the Zen Loader with an async loading thread."/>
 	<Option Name="DisableEditorDomainForZenLoaderTests" DefaultValue="true" Description="Explicitly disables editordomain (-noeditordomain) for zen loader tests"/>
 
 	<!-- Zen Loader tests specifc arguments -->
@@ -186,6 +186,7 @@
 	<Property Name="GauntletArgs" Value="-build=editor -tempdir=&quot;$(RootDir)/Tests&quot;"/>
 	<Property Name="GauntletArgs" Value="$(GauntletArgs) -branch=$(Branch) -ECBranch=&quot;$(Branch)&quot; -ECChangelist=&quot;$(Change)&quot;"/>
 	<Property Name="GauntletArgs" Value="$(GauntletArgs) -verbose -maxduration=7200"/>
+	<Property Name="GauntletArgs" Value="$(GauntletArgs) -verbose -maxduration=18000" If="$(WithTSAN)"/>
 
 	<!-- Gauntlet arguments for profile runs -->
 	<Property Name="ProfileArgs" Value="-tracefiletrunc"/>
