- Updates EditorPerf jobs to only on Cold jobs when running Warm jobs so we don't run all cold jobs sequentially
- Fixes missing condition preventing TSAN tests from running with the correct timeout

diff --git a/Engine/Build/EditorPerf.xml b/Engine/Build/EditorPerf.xml
--- a/Engine/Build/EditorPerf.xml
+++ b/Engine/Build/EditorPerf.xml
@@ -185,7 +185,7 @@
 	<!-- Gauntlet arguments for all runs -->
 	<Property Name="GauntletArgs" Value="-build=editor -tempdir=&quot;$(RootDir)/Tests&quot;"/>
 	<Property Name="GauntletArgs" Value="$(GauntletArgs) -branch=$(Branch) -ECBranch=&quot;$(Branch)&quot; -ECChangelist=&quot;$(Change)&quot;"/>
-	<Property Name="GauntletArgs" Value="$(GauntletArgs) -verbose -maxduration=7200"/>
+	<Property Name="GauntletArgs" Value="$(GauntletArgs) -verbose -maxduration=7200" If="!$(WithTSAN)"/>
 	<Property Name="GauntletArgs" Value="$(GauntletArgs) -verbose -maxduration=18000" If="$(WithTSAN)"/>
 
 	<!-- Gauntlet arguments for profile runs -->
@@ -376,12 +376,16 @@
 				
 				<Agent Name="EditorPerf Agent $(Platform) $(Project)" Type="EditorPerf$(Platform)">
 					<ForEach Name="Phase" Values="$(Phases)">
+						<!-- Make Warm runs depend on previously defined nodes -->
+						<Property Name="RequiredNodes" Value="" If="'$(Phase)' == 'Cold'"/>
+						<Property Name="RequiredNodes" Value="$(RunNodeList)" If="'$(Phase)' == 'Warm'"/>
+						
 						<Do If="$(EnableBasePhaseTests)">
 							<ForEach Name="Iteration" Values="$(Iterations)">
 								<Property Name="IterationSuffix" Value=""/>
 								<Property Name="IterationSuffix" Value="$(Iteration)" If="'$(Iteration)' != '1'"/>
 								<Property Name="RunNodeName" Value="Run $(OutputName) $(Phase)$(IterationSuffix)"/>
-								<Node Name="$(RunNodeName)" Requires="$(RunNodeRequires);$(RunNodeList)">
+								<Node Name="$(RunNodeName)" Requires="$(RunNodeRequires);$(RequiredNodes)">
 									<Expand Name="RunTest" TestNameSuffix="NoZenLoader" Platform="$(Platform)" ProjectName="$(Project)" ProjectTarget="$(ProjectTarget)" NodeName="$(RunNodeName)" Phase="$(Phase)" Iteration="$(Iteration)" Args="$($(Phase)Args)" EditorArgs="$(ProjectCombinedArgs) $(NoZenLoaderTestArgs)"/>
 								</Node>
 								<Property Name="RunNodeList" Value="$(RunNodeList);$(RunNodeName)"/>
@@ -393,7 +397,7 @@
 						<Do>
 							<!--Zen Loader TestRun if enabled--> 
 							<Property Name="RunNodeZenName" Value="Run $(OutputName) $(Phase) ZenLoader"/>
-							<Node Name="$(RunNodeZenName)" Requires="$(RunNodeRequires);$(RunNodeList)" If="$(EnableZenLoaderTest)">
+							<Node Name="$(RunNodeZenName)" Requires="$(RunNodeRequires);$(RequiredNodes)" If="$(EnableZenLoaderTest)">
 								<Expand Name="RunTest" TestNameSuffix="ZenLoader" Platform="$(Platform)" ProjectName="$(Project)" ProjectTarget="$(ProjectTarget)" NodeName="$(RunNodeZenName)" Phase="$(Phase)" Iteration="1" Args="$($(Phase)Args)" EditorArgs="$(ProjectCombinedArgs) $(ZenLoaderTestArgs)"/>
 							</Node>
 							<Property Name="RunNodeList" Value="$(RunNodeList);$(RunNodeZenName)" If="$(EnableZenLoaderTest)"/>
@@ -401,7 +405,7 @@
 							
 							<!--Zen Loader Async TestRun if enabled--> 
 							<Property Name="RunNodeZenAsyncName" Value="Run $(OutputName) $(Phase) ZenLoader Async"/>
-							<Node Name="$(RunNodeZenAsyncName)" Requires="$(RunNodeRequires);$(RunNodeList)" If="$(EnableZenLoaderAsyncThreadTest)">
+							<Node Name="$(RunNodeZenAsyncName)" Requires="$(RunNodeRequires);$(RequiredNodes)" If="$(EnableZenLoaderAsyncThreadTest)">
 								<Expand Name="RunTest" TestNameSuffix="ZenLoaderAsync" Platform="$(Platform)" ProjectName="$(Project)" ProjectTarget="$(ProjectTarget)" NodeName="$(RunNodeZenAsyncName)" Phase="$(Phase)" Iteration="1" Args="$($(Phase)Args)" EditorArgs="$(ProjectCombinedArgs) $(ZenLoaderAsyncTestArgs)"/>
 							</Node>
 							<Property Name="RunNodeList" Value="$(RunNodeList);$(RunNodeZenAsyncName)" If="$(EnableZenLoaderAsyncThreadTest)"/>
@@ -409,7 +413,7 @@
 
 							<!--Zen Store DDC TestRun if enabled--> 
 							<Property Name="RunNodeZenStoreDDCName" Value="Run $(OutputName) $(Phase) ZenLoader Store DDC"/>
-							<Node Name="$(RunNodeZenStoreDDCName)" Requires="$(RunNodeRequires);$(RunNodeList)" If="$(EnableZenStoreDDCTest)">
+							<Node Name="$(RunNodeZenStoreDDCName)" Requires="$(RunNodeRequires);$(RequiredNodes)" If="$(EnableZenStoreDDCTest)">
 								<Expand Name="RunTest" TestNameSuffix="ZenStoreAsyncDDC" Platform="$(Platform)" ProjectName="$(Project)" ProjectTarget="$(ProjectTarget)" NodeName="$(RunNodeZenStoreDDCName)" Phase="$(Phase)" Iteration="1" Args="$(ZenStoreDDCPhaseArgs)" EditorArgs="$(ProjectCombinedArgs) $(ZenStoreDDCTestArgs)"/>
 							</Node>
 							<Property Name="RunNodeList" Value="$(RunNodeList);$(RunNodeZenStoreDDCName)" If="$(EnableZenStoreDDCTest)"/>
@@ -420,7 +424,7 @@
 						<Do If="'$(Phase)' == 'Warm'">
 							<!--Legacy LazyLoad TestRun if enabled--> 
 							<Property Name="RunNodeLazyLoadName" Value="Run $(OutputName) $(Phase) LazyLoad NoZenLoader"/>
-							<Node Name="$(RunNodeLazyLoadName)" Requires="$(RunNodeRequires);$(RunNodeList)" If="$(EnableLazyLoadNoZenLoaderTest)">
+							<Node Name="$(RunNodeLazyLoadName)" Requires="$(RunNodeRequires);$(RequiredNodes)" If="$(EnableLazyLoadNoZenLoaderTest)">
 								<Expand Name="RunTest" TestNameSuffix="NoZenLoader" Platform="$(Platform)" ProjectName="$(Project)" ProjectTarget="$(ProjectTarget)" NodeName="$(RunNodeLazyLoadName)" Phase="$(Phase)" Iteration="1" Args="$($(Phase)Args)" EditorArgs="$(ProjectCombinedArgs) $(LazyLoadTestArgs) $(NoZenLoaderTestArgs)"/>
 							</Node>
 							<Property Name="RunNodeList" Value="$(RunNodeList);$(RunNodeLazyLoadName)" If="$(EnableLazyLoadNoZenLoaderTest)"/>
@@ -428,7 +432,7 @@
 							
 							<!--Zen LazyLoad TestRun if enabled--> 
 							<Property Name="RunNodeLazyLoadZenName" Value="Run $(OutputName) $(Phase) LazyLoad ZenLoader"/>
-							<Node Name="$(RunNodeLazyLoadZenName)" Requires="$(RunNodeRequires);$(RunNodeList)" If="$(EnableLazyLoadZenLoaderTest)">
+							<Node Name="$(RunNodeLazyLoadZenName)" Requires="$(RunNodeRequires);$(RequiredNodes)" If="$(EnableLazyLoadZenLoaderTest)">
 								<Expand Name="RunTest" TestNameSuffix="ZenLoader" Platform="$(Platform)" ProjectName="$(Project)" ProjectTarget="$(ProjectTarget)" NodeName="$(RunNodeLazyLoadZenName)" Phase="$(Phase)" Iteration="1" Args="$($(Phase)Args)" EditorArgs="$(ProjectCombinedArgs) $(LazyLoadZenTestArgs)"/>
 							</Node>
 							<Property Name="RunNodeList" Value="$(RunNodeList);$(RunNodeLazyLoadZenName)" If="$(EnableLazyLoadZenLoaderTest)"/>
