Add platform support for the run shader type Horde job.

diff --git a/Engine/Build/ShaderTypeRegressionTest.xml b/Engine/Build/ShaderTypeRegressionTest.xml
--- a/Engine/Build/ShaderTypeRegressionTest.xml
+++ b/Engine/Build/ShaderTypeRegressionTest.xml
@@ -13,15 +13,16 @@
 	<Property Name="BuildName" Value="$(EscapedBranch)-CL-$(Change)$(PreflightSuffix)"/>
 	<Option Name="CompileCategory" DefaultValue="Editors" Description="Category string for Compile Label"/>
 	<Option Name="RunCategory" DefaultValue="Test" Description="Category string for Run Label"/>
-	
+
 	<Property Name="AnalyticsArg" Value=""/>
 	<Option Name="SendAnalytics" DefaultValue="false" Description="Whether or not to shader stat analytics."/>
 	<Do If="$(SendAnalytics)">
 		<Property Name="AnalyticsArg" Value="-analytics"/>
 	</Do>
-	
-	<Property Name="CommandletArguments" Value="-targetplatform=Windows $(AnalyticsArg)"/>
-	
+
+	<!-- Platform setup-->
+	<Option Name="Platforms" DefaultValue="Windows" Restrict=".+" Description="Plaform to run the commandlet"/>
+
 	<!-- Projects setup-->
 	<Option Name="Projects" DefaultValue="false" Restrict=".+" Description="The projects to test"/>
 	<Error Message="No project names specified" If="'$(Projects)' == 'false'"/>
@@ -54,20 +55,25 @@
 			
 			<!-- Not sure if I actually need this list -->
 			<Property Name="RunNodeList" Value=""/>
-			
-			<!-- Run the commandlet itself. -->
-			<Property Name="RunNodeName" Value="$(Project) Dump Material Shader Types"/>
-			<Node Name="$(RunNodeName)" Requires="$(RunNodeRequires);$(RunNodeList)">
-				<Commandlet Name="DumpMaterialShaderTypes" Project="$(Project)" Arguments="$(CommandletArguments)"/>
-			</Node>
-			<Property Name="RunNodeList" Value="$(RunNodeList);$(RunNodeName)"/>
+
+			<!-- ForEach platforms, run the commandlet -->
+			<ForEach Name="Platform" Values="$(Platforms)">
+				<Property Name="RunNodeName" Value="$(Project) Dump Material Shader Types"/>
+				<!-- Run the commandlet itself. -->
+				<Node Name="$(RunNodeName)" Requires="$(RunNodeRequires);$(RunNodeList)">
+					<Commandlet Name="DumpMaterialShaderTypes" Project="$(Project)" Arguments="-targetplatform=$(Platform)"/>
+				</Node>
+				<Property Name="RunNodeList" Value="$(RunNodeList);$(RunNodeName)"/>
+			</ForEach>
+			<!-- Platforms -->
 
 			<Label Category="$(RunCategory)" Name="EP Run $(Project) $(EditorPlatform)" Requires="$(RunNodeRequires);$(RunNodeList)"/>
 			<Property Name="RunGraphNodes" Value="$(RunGraphNodes);$(RunNodeList)"/>
 		</Agent>
 
-	</ForEach> <!-- Projects -->
-	
+	</ForEach>
+	<!-- Projects -->
+
 	<!-- These Nodes run the Editor and gather profile data -->
 	<Aggregate Name="Shader Type Regression" Requires="$(RunGraphNodes)"/>
 
