Change EditorPlatform type so that ShaderTypeRegression script works on different stream projects.

diff --git a/Engine/Build/ShaderTypeRegressionTest.xml b/Engine/Build/ShaderTypeRegressionTest.xml
--- a/Engine/Build/ShaderTypeRegressionTest.xml
+++ b/Engine/Build/ShaderTypeRegressionTest.xml
@@ -31,7 +31,7 @@
 	<Property Name="RunGraphNodes" Value=""/>
 	
 	<!-- This will be set to the names of all nodes we want to execute based on specified platforms/options -->
-	<Property Name="EditorPlatform" Value="Win64"/>
+	<Property Name="EditorPlatform" Value="Win64Generic;Win64"/>
 
 	<!-- ForEach Project to run, generate the graph nodes -->
 	<ForEach Name="Project" Values="$(Projects)">
