Fixing Build Graph errors around EditorPlatform

#rnx

diff --git a/Engine/Build/ShaderTypeRegressionTest.xml b/Engine/Build/ShaderTypeRegressionTest.xml
--- a/Engine/Build/ShaderTypeRegressionTest.xml
+++ b/Engine/Build/ShaderTypeRegressionTest.xml
@@ -31,7 +31,8 @@
 	<Property Name="RunGraphNodes" Value=""/>
 	
 	<!-- This will be set to the names of all nodes we want to execute based on specified platforms/options -->
-	<Property Name="EditorPlatform" Value="Win64Generic;Win64"/>
+	<Property Name="EditorPlatform" Value="Win64"/>
+	<Property Name="AgentType" Value="Win64Generic;Win64"/>
 
 	<!-- ForEach Project to run, generate the graph nodes -->
 	<ForEach Name="Project" Values="$(Projects)">
@@ -42,7 +43,7 @@
 
 		<!-- Compile the Editor -->
 		<Property Name="CompileNodeName" Value="Compile $(Project) $(EditorPlatform)"/>
-		<Agent Name="Compile Editor Agent $(EditorPlatform) $(Project)" Type="$(EditorPlatform)">
+		<Agent Name="Compile Editor Agent $(EditorPlatform) $(Project)" Type="$(AgentType)">
 			<Node Name="$(CompileNodeName)">
 				<Compile Target="ShaderCompileWorker" Platform="$(EditorPlatform)" Configuration="Development"/>
 				<Compile Target="CrashReportClientEditor" Platform="$(EditorPlatform)" Project="$(Project)" Configuration="Shipping"/>
@@ -78,3 +79,4 @@
 	<Aggregate Name="Shader Type Regression" Requires="$(RunGraphNodes)"/>
 
 </BuildGraph>
+
