- Removes EditorWorkflow nodes from EditorPerf.xml as those jobs are not run in CI and have been deprecated.

- Makes 'warm' jobs only depend on the specific 'cold' job for their config, rather than depending on all previous cold jobs.

- Makes TSAN runs only run on warm runs as cold is too slow. TSAN jobs will thus, build and run a cold job without TSAN enabled to prime the DDC and then run a 'warm' TSAN enabled job. We do two steps for DDC priming instead of using shared cache or incremental DDC so that the TSAN jobs are more reliable in terms of the code they execute.

#rnx

diff --git a/Engine/Build/EditorPerf.xml b/Engine/Build/EditorPerf.xml
--- a/Engine/Build/EditorPerf.xml
+++ b/Engine/Build/EditorPerf.xml
@@ -42,7 +42,7 @@
 	<Property Name="EditorCompileArgs" Value=""/>
 	<Property Name="EditorCompileArgs" Value="-BuildPlugin=$(ExtraPluginsPlusDelims)" If="'$(ExtraPluginsPlusDelims)' != ''"/>
 	<Property Name="EditorCompileArgs" Value="$(EditorCompileArgs) -nodebuginfo" If="$(NoDebugInfo)"/>
-	<Property Name="EditorCompileArgs" Value="$(EditorCompileArgs) -EnableTSAN -Monolithic" If="$(WithTSAN)"/>
+	<Property Name="TSANEditorCompileArgs" Value="-EnableTSAN -Monolithic"/>
 	<Property Name="EditorCompileArgs" Value="$(EditorCompileArgs) $(AdditionalEditorCompileArgs)" If="'$(AdditionalEditorCompileArgs)' != ''"/>
 
 	<!-- Get shared properties -->
@@ -276,48 +276,14 @@
 		<Property Name="GenerateTSANReportNodeList" Value="$(GenerateTSANReportNodeList); $(GenerateTSANReportNodeName)"/>
 	</Macro>
 	
-	<!-- The main editor workflows test running macro logic that is expanded with platform and dependency node names -->
-	<Macro Name="RunEditorWorkflowsTest" Arguments="Platform;ProjectName;NodeName">
-		<Property Name="TestName" Value="Editor.Workflows"/>
-		<Property Name="TestFilterArgs" Value="-runtest=$(TestName)"/>
-
-		<Property Name="EditorArgs" Value=""/>
-		<Property Name="TestOutputDir" Value="$(AutomationOutputDir)/$(ProjectName)$(Platform)EditorWorkflows"/>
-		<Property Name="TraceFilePath" Value="$(TraceDir)/$(TestName)/$(ProjectName)$(Platform)EditorWorkflows.utrace" />
-
-		<Property Name="LogOutputParams" Value="-logdir=&quot;$(TestOutputDir)&quot;"/>
-		<Property Name="ReportParams" Value="-WriteTestResultsForHorde -HordeTestDataKey=&quot;$(NodeName)$(PreflightSuffix)&quot;"/>
-		<Property Name="AllGauntletArgs" Value="-Project=$(ProjectName) -Platform=$(Platform) $(GauntletArgs) $(ReportParams) $(LogOutputParams)"/>
-
-		<Property Name="TestArg" Value="-test=EditorTests" />
-		<Property Name="RunArgs" Value="$(TestArg) $(TestFilterArgs) -editorargs=&quot;$(EditorArgs)&quot;"/>
-
-		<Command Name="RunUnreal" Arguments="$(AllGauntletArgs) $(ProfileArgs) $(RunArgs) -sessionlabel=&quot;$(TestName)&quot; -tracefile=&quot;$(TraceFilePath)&quot;" />
-	</Macro>
-	
 	<Macro Name="UploadTestTraces" Arguments="Platform;ProjectName">		
 		<!-- Save all files named for this project in the trace directory -->
 		<Property Name="TestOutputDir" Value="$(AutomationOutputDir)/$(ProjectName)$(Platform)EditorPerf"/>
 		<Copy From="$(TraceDir)/$(ProjectName)$(Platform)EditorPerf*.*" To="$(TestOutputDir)"/>
 	</Macro>
 
-	<Macro Name="UploadEditorWorkflowsPerformanceDataAndTraces" Arguments="Platform;ProjectName">
-		<Property Name="TestOutputDir" Value="$(AutomationOutputDir)/$(ProjectName)$(Platform)EditorWorkflows"/>
-		<Property Name="TraceFilePath" Value="$(TraceDir)/Editor.Workflows/$(ProjectName)$(Platform)EditorWorkflows.utrace" />
-
-		<!-- Translate the utrace file into csv data -->
-		<Commandlet Project="$(ProjectName)" Name="SummarizeTraceEditorWorkflows" Arguments="-inputfile=&quot;$(TraceFilePath)&quot; -testname=&quot;Editor.Workflows&quot; $(CommandletArgs)" />
-
-		<!-- Copy the utrace file and the generated csv performance data to the output directory -->
-		<Copy From="$(TraceDir)/Editor.Workflows/*" To="$(TestOutputDir)/Editor.Workflows" Overwrite="true"/>
-
-		<!-- Copy "index.json" that captures the functional results to the output directory -->
-		<Copy From="$(RootDir)/Tests/TestReport/index.json" To="$(TestOutputDir)/Editor.Workflows/index.json" Overwrite="true"/>
-	</Macro>
-
 	<!-- These will be set to the names of all nodes we want to execute based on specified platforms/options -->
 	<Property Name="RunGraphNodes" Value=""/>
-	<Property Name="RunEditorWorkflowsGraphNodes" Value=""/>
 
 	<!-- Mimic of the Runner that retrieves the stashed telemetry file and copies it to the depot location -->
 	<Macro Name="BaselineResult" Arguments="Platform;Project;Phase">
@@ -351,15 +317,26 @@
 
 			<!-- Compile the Editor. Set up compilation nodes for all projects that have been defined even if we aren't going to test them. This lets the tested projects depend on non-tested ones. -->
 			<Property Name="CompileNodeName" Value="Compile $(Project) $(Platform)" />
+			<Property Name="TSANCompileNodeName" Value="Compile $(Project) $(Platform) with TSAN" />
 			<Property Name="AgentModifier" Value=""/>
 			<Property Name="AgentModifier" Value=" $(Project)" If="!$(SingleAgentCompile)"/>
 			<Agent Name="Compile Editor Agent $(Platform)$(AgentModifier" Type="$($(Platform)CompileAgentType)">
+				<!-- 
+				If running with TSAN we compile twice, once without TSAN enabled (used for cold runs) and once with TSAN enabled (used for warm runs).
+				TSAN adds too much overhead to enable on cold runs, but we need a cold run to execute to prime the cache for the warm run so we run both.
+				-->
 				<Node Name="$(CompileNodeName)">
 					<SetVersion Change="$(Change)" Branch="$(EscapedBranch)" If="$(IsBuildMachine)"/>
 					<Compile Target="ShaderCompileWorker" Platform="$(Platform)" Configuration="Development"/>
 					<Compile Target="CrashReportClientEditor" Platform="$(Platform)" Project="$(ProjectTarget)" Configuration="Shipping"/>
 					<Compile Target="$(EditorTarget)" Platform="$(Platform)" Project="$(ProjectTarget)" Configuration="Development" Arguments="$(EditorCompileArgs)"/>
 				</Node>
+				<Node Name="$(TSANCompileNodeName)" If="$(WithTSAN)">
+					<SetVersion Change="$(Change)" Branch="$(EscapedBranch)" If="$(IsBuildMachine)"/>
+					<Compile Target="ShaderCompileWorker" Platform="$(Platform)" Configuration="Development"/>
+					<Compile Target="CrashReportClientEditor" Platform="$(Platform)" Project="$(ProjectTarget)" Configuration="Shipping"/>
+					<Compile Target="$(EditorTarget)" Platform="$(Platform)" Project="$(ProjectTarget)" Configuration="Development" Arguments="$(EditorCompileArgs) $(TSANEditorCompileArgs)"/>
+				</Node>
 			</Agent>
 			<Label Category="$(CompileCategory)" Name="$(Project) Editor Compile" Requires="$(CompileNodeName)"/>
 			<!-- Store off the name of the compile node so that we can use look it up by Project name from UseBinariesFrom. -->
@@ -367,25 +344,25 @@
 
 			<!-- If the project will be run, set up the test nodes -->
 			<Do If="ContainsItem('$(Projects)', '$(Project)', ';')=='true'">
-				<!-- Run the Tests -->
-				<Property Name="RunNodeRequires" Value="$(CompileNodeName)"/>
-				<Property Name="RunNodeRequires" Value="$($($(Project)__UseBinariesFrom)__CompileNodeName)" If="'$($(Project)__UseBinariesFrom)'!=''"/>
-				<Property Name="RunNodeList" Value=""/>
-				<!--Each node in this the nested loops below, they need to depend on predecessors--> 
+				<!--Each 'Warm' node in this the nested loops below depend on their 'Cold' predecessors. The nodes themselves are otherwise distinct from one another --> 
+				<Property Name="RunNodeList" Value=""/>		
 				<Property Name="UploadNodeList" Value=""/>
-				
 				<Agent Name="EditorPerf Agent $(Platform) $(Project)" Type="EditorPerf$(Platform)">
 					<ForEach Name="Phase" Values="$(Phases)">
-						<!-- Make Warm runs depend on previously defined nodes -->
-						<Property Name="RequiredNodes" Value="" If="'$(Phase)' == 'Cold'"/>
-						<Property Name="RequiredNodes" Value="$(RunNodeList)" If="'$(Phase)' == 'Warm'"/>
-						
+
+						<Property Name="RunNodeCompilationRequires" Value="$(CompileNodeName)"/>
+						<Property Name="RunNodeCompilationRequires" Value="$($($(Project)__UseBinariesFrom)__CompileNodeName)" If="'$($(Project)__UseBinariesFrom)'!=''"/>
+						<!-- Only use a TSAN enabled builds on warm builds (if using TSAN at all) -->
+						<Property Name="RunNodeCompilationRequires" Value="$(TSANCompileNodeName)" If="$(WithTSAN) and '$(Phase)' == 'Warm'"/>
+
 						<Do If="$(EnableBasePhaseTests)">
 							<ForEach Name="Iteration" Values="$(Iterations)">
 								<Property Name="IterationSuffix" Value=""/>
 								<Property Name="IterationSuffix" Value="$(Iteration)" If="'$(Iteration)' != '1'"/>
 								<Property Name="RunNodeName" Value="Run $(OutputName) $(Phase)$(IterationSuffix)"/>
-								<Node Name="$(RunNodeName)" Requires="$(RunNodeRequires);$(RequiredNodes)">
+								<Property Name="RunNodeRequires" Value="$(RunNodeCompilationRequires)"/>
+								<Property Name="RunNodeRequires" Value="$(RunNodeCompilationRequires); Run $(OutputName) Cold$(IterationSuffix)" If="'$(Phase)' == 'Warm'"/>
+								<Node Name="$(RunNodeName)" Requires="$(RunNodeRequires)">
 									<Expand Name="RunTest" TestNameSuffix="NoZenLoader" Platform="$(Platform)" ProjectName="$(Project)" ProjectTarget="$(ProjectTarget)" NodeName="$(RunNodeName)" Phase="$(Phase)" Iteration="$(Iteration)" Args="$($(Phase)Args)" EditorArgs="$(ProjectCombinedArgs) $(NoZenLoaderTestArgs)"/>
 								</Node>
 								<Property Name="RunNodeList" Value="$(RunNodeList);$(RunNodeName)"/>
@@ -397,7 +374,9 @@
 						<Do>
 							<!--Zen Loader TestRun if enabled--> 
 							<Property Name="RunNodeZenName" Value="Run $(OutputName) $(Phase) ZenLoader"/>
-							<Node Name="$(RunNodeZenName)" Requires="$(RunNodeRequires);$(RequiredNodes)" If="$(EnableZenLoaderTest)">
+							<Property Name="RunNodeZenRequires" Value="$(RunNodeCompilationRequires)"/>
+							<Property Name="RunNodeZenRequires" Value="$(RunNodeCompilationRequires); Run $(OutputName) Cold ZenLoader" If="'$(Phase)' == 'Warm'"/>
+							<Node Name="$(RunNodeZenName)" Requires="$(RunNodeZenRequires)" If="$(EnableZenLoaderTest)">
 								<Expand Name="RunTest" TestNameSuffix="ZenLoader" Platform="$(Platform)" ProjectName="$(Project)" ProjectTarget="$(ProjectTarget)" NodeName="$(RunNodeZenName)" Phase="$(Phase)" Iteration="1" Args="$($(Phase)Args)" EditorArgs="$(ProjectCombinedArgs) $(ZenLoaderTestArgs)"/>
 							</Node>
 							<Property Name="RunNodeList" Value="$(RunNodeList);$(RunNodeZenName)" If="$(EnableZenLoaderTest)"/>
@@ -405,7 +384,9 @@
 							
 							<!--Zen Loader Async TestRun if enabled--> 
 							<Property Name="RunNodeZenAsyncName" Value="Run $(OutputName) $(Phase) ZenLoader Async"/>
-							<Node Name="$(RunNodeZenAsyncName)" Requires="$(RunNodeRequires);$(RequiredNodes)" If="$(EnableZenLoaderAsyncThreadTest)">
+							<Property Name="RunNodeZenAsyncRequires" Value="$(RunNodeCompilationRequires)"/>
+							<Property Name="RunNodeZenAsyncRequires" Value="$(RunNodeCompilationRequires); Run $(OutputName) Cold ZenLoader Async" If="'$(Phase)' == 'Warm'"/>
+							<Node Name="$(RunNodeZenAsyncName)" Requires="$(RunNodeZenAsyncRequires)" If="$(EnableZenLoaderAsyncThreadTest)">
 								<Expand Name="RunTest" TestNameSuffix="ZenLoaderAsync" Platform="$(Platform)" ProjectName="$(Project)" ProjectTarget="$(ProjectTarget)" NodeName="$(RunNodeZenAsyncName)" Phase="$(Phase)" Iteration="1" Args="$($(Phase)Args)" EditorArgs="$(ProjectCombinedArgs) $(ZenLoaderAsyncTestArgs)"/>
 							</Node>
 							<Property Name="RunNodeList" Value="$(RunNodeList);$(RunNodeZenAsyncName)" If="$(EnableZenLoaderAsyncThreadTest)"/>
@@ -413,7 +394,9 @@
 
 							<!--Zen Store DDC TestRun if enabled--> 
 							<Property Name="RunNodeZenStoreDDCName" Value="Run $(OutputName) $(Phase) ZenLoader Store DDC"/>
-							<Node Name="$(RunNodeZenStoreDDCName)" Requires="$(RunNodeRequires);$(RequiredNodes)" If="$(EnableZenStoreDDCTest)">
+							<Property Name="RunNodeZenStoreDDCRequires" Value="$(RunNodeCompilationRequires)"/>
+							<Property Name="RunNodeZenStoreDDCRequires" Value="$(RunNodeCompilationRequires); Run $(OutputName) Cold ZenLoader Store DDC" If="'$(Phase)' == 'Warm'"/>
+							<Node Name="$(RunNodeZenStoreDDCName)" Requires="$(RunNodeZenStoreDDCRequires)" If="$(EnableZenStoreDDCTest)">
 								<Expand Name="RunTest" TestNameSuffix="ZenStoreAsyncDDC" Platform="$(Platform)" ProjectName="$(Project)" ProjectTarget="$(ProjectTarget)" NodeName="$(RunNodeZenStoreDDCName)" Phase="$(Phase)" Iteration="1" Args="$(ZenStoreDDCPhaseArgs)" EditorArgs="$(ProjectCombinedArgs) $(ZenStoreDDCTestArgs)"/>
 							</Node>
 							<Property Name="RunNodeList" Value="$(RunNodeList);$(RunNodeZenStoreDDCName)" If="$(EnableZenStoreDDCTest)"/>
@@ -424,7 +407,9 @@
 						<Do If="'$(Phase)' == 'Warm'">
 							<!--Legacy LazyLoad TestRun if enabled--> 
 							<Property Name="RunNodeLazyLoadName" Value="Run $(OutputName) $(Phase) LazyLoad NoZenLoader"/>
-							<Node Name="$(RunNodeLazyLoadName)" Requires="$(RunNodeRequires);$(RequiredNodes)" If="$(EnableLazyLoadNoZenLoaderTest)">
+							<Property Name="RunNodeLazyLoadRequires" Value="$(RunNodeCompilationRequires)"/>
+							<Property Name="RunNodeLazyLoadRequires" Value="$(RunNodeCompilationRequires); Run $(OutputName) Cold LazyLoad NoZenLoader" If="'$(Phase)' == 'Warm'"/>
+							<Node Name="$(RunNodeLazyLoadName)" Requires="$(RunNodeLazyLoadRequires)" If="$(EnableLazyLoadNoZenLoaderTest)">
 								<Expand Name="RunTest" TestNameSuffix="NoZenLoader" Platform="$(Platform)" ProjectName="$(Project)" ProjectTarget="$(ProjectTarget)" NodeName="$(RunNodeLazyLoadName)" Phase="$(Phase)" Iteration="1" Args="$($(Phase)Args)" EditorArgs="$(ProjectCombinedArgs) $(LazyLoadTestArgs) $(NoZenLoaderTestArgs)"/>
 							</Node>
 							<Property Name="RunNodeList" Value="$(RunNodeList);$(RunNodeLazyLoadName)" If="$(EnableLazyLoadNoZenLoaderTest)"/>
@@ -432,7 +417,9 @@
 							
 							<!--Zen LazyLoad TestRun if enabled--> 
 							<Property Name="RunNodeLazyLoadZenName" Value="Run $(OutputName) $(Phase) LazyLoad ZenLoader"/>
-							<Node Name="$(RunNodeLazyLoadZenName)" Requires="$(RunNodeRequires);$(RequiredNodes)" If="$(EnableLazyLoadZenLoaderTest)">
+							<Property Name="RunNodeLazyLoadZenRequires" Value="$(RunNodeCompilationRequires)"/>
+							<Property Name="RunNodeLazyLoadZenRequires" Value="$(RunNodeCompilationRequires); Run $(OutputName) Cold LazyLoad ZenLoader" If="'$(Phase)' == 'Warm'"/>
+							<Node Name="$(RunNodeLazyLoadZenName)" Requires="$(RunNodeLazyLoadZenRequires);" If="$(EnableLazyLoadZenLoaderTest)">
 								<Expand Name="RunTest" TestNameSuffix="ZenLoader" Platform="$(Platform)" ProjectName="$(Project)" ProjectTarget="$(ProjectTarget)" NodeName="$(RunNodeLazyLoadZenName)" Phase="$(Phase)" Iteration="1" Args="$($(Phase)Args)" EditorArgs="$(ProjectCombinedArgs) $(LazyLoadZenTestArgs)"/>
 							</Node>
 							<Property Name="RunNodeList" Value="$(RunNodeList);$(RunNodeLazyLoadZenName)" If="$(EnableLazyLoadZenLoaderTest)"/>
@@ -448,7 +435,7 @@
 						<Property Name="UploadNodeList" Value="$(RunNodeUploadPhaseTraces)"/>
 					</Do>
 				</Agent>
-				<Label Category="$(RunCategory)" Name="$(Project) Editor Run" Requires="$(RunNodeRequires);$(RunNodeList);$(UploadNodeList)" UgsBadge="$(Project) Editor Run"/>
+				<Label Category="$(RunCategory)" Name="$(Project) Editor Run" Requires="$(RunNodeList);$(UploadNodeList)" UgsBadge="$(Project) Editor Run"/>
 				<Property Name="RunGraphNodes" Value="$(RunGraphNodes);$(RunNodeList);$(GenerateTSANReportNodeList);$(UploadNodeList)"/>
 
 				<!--Baseline the Results--> 
@@ -469,58 +456,6 @@
 
 	</ForEach> <!-- Platforms -->
 
-	<!-- Generate graph nodes for each platform to run -->
-	<ForEach Name="Platform" Values="$(Platforms)">
-
-		<Property Name="ProjectName" Value="EngineTest" />
-		<!-- Define output name for the EditorWorkflows test -->
-		<Property Name="EditorWorkflowsOutputName" Value="$(ProjectName)$(Platform)EditorWorkflows"/>
-		<Property Name="CompileNodeName" Value="Compile $(ProjectName) $(Platform) EditorWorkflows"/>
-		<Property Name="AgentModifier" Value=""/>
-		<Property Name="AgentModifier" Value=" $(ProjectName)" If="!$(SingleAgentCompile)"/>
-
-		<!-- Compile Agent configuration -->
-		<Agent Name="Compile Editor Agent $(Platform)$(AgentModifier)" Type="$($(Platform)CompileAgentType)">
-			<Node Name="$(CompileNodeName)">
-				<SetVersion Change="$(Change)" Branch="$(EscapedBranch)" If="$(IsBuildMachine)"/>
-				<Compile Target="ShaderCompileWorker" Platform="$(Platform)" Configuration="Development"/>
-				<Compile Target="CrashReportClientEditor" Platform="$(Platform)" Project="$(ProjectName)" Configuration="Shipping"/>
-				<Compile Target="$(ProjectName)Editor" Platform="$(Platform)" Project="$(ProjectName)" Configuration="Development" Arguments="$(EditorCompileArgs)"/>
-			</Node>
-		</Agent>
-		<Label Category="$(CompileCategory)" Name="$(ProjectName) Editor Compile" Requires="$(CompileNodeName)"/>
-
-		<!-- Set up properties for run and upload nodes -->
-		<Property Name="RunEditorWorkflowsNodeRequires" Value="$(CompileNodeName)"/>
-		<Property Name="UploadEditorWorkflowsNodeList" Value=""/>
-		<Property Name="RunEditorWorkflowsNodeList" Value=""/>
-
-		<!-- Configure EditorWorkflows Agent -->
-		<Agent Name="EditorWorkflows Agent $(Platform) $(ProjectName)" Type="EditorPerf$(Platform)">
-			<Property Name="RunEditorWorkflowsNodeName" Value="Run $(EditorWorkflowsOutputName)" />
-			<Node Name="$(RunEditorWorkflowsNodeName)" Requires="$(RunEditorWorkflowsNodeRequires);$(RunEditorWorkflowsNodeList)">
-				<Expand Name="RunEditorWorkflowsTest" Platform="$(Platform)" ProjectName="$(ProjectName)" NodeName="$(RunEditorWorkflowsNodeName)"/>
-			</Node>
-
-			<Property Name="RunEditorWorkflowsNodeList" Value="$(RunEditorWorkflowsNodeList);$(RunEditorWorkflowsNodeName)" />
-			<Property Name="RunNodeUploadTraces" Value="Upload $(EditorWorkflowsOutputName) Traces"/>
-
-			<!-- Set up Node for uploading traces -->
-			<Node Name="$(RunNodeUploadTraces)">
-				<Expand Name="UploadEditorWorkflowsPerformanceDataAndTraces" Platform="$(Platform)" ProjectName="$(ProjectName)"/>
-			</Node>
-
-			<!-- Update UploadEditorWorkflowsNodeList -->
-			<Property Name="UploadEditorWorkflowsNodeList" Value="$(UploadEditorWorkflowsNodeList);$(RunNodeUploadTraces)"/>
-		</Agent>
-		<Label Category="$(RunCategory)" Name="$(ProjectName) Editor Run" Requires="$(RunEditorWorkflowsNodeRequires);$(RunEditorWorkflowsNodeList);$(UploadEditorWorkflowsNodeList)" UgsBadge="$(ProjectName) Editor Run"/>
-		<Property Name="RunEditorWorkflowsGraphNodes" Value="$(RunEditorWorkflowsGraphNodes);$(RunEditorWorkflowsNodeList);$(UploadEditorWorkflowsNodeList)"/>
-	</ForEach>
-	<!-- End ForEach Platforms -->
-
-	<!-- These Nodes run the Editor Workflows test and gather profile data -->
-	<Aggregate Name="Run EditorWorkflows Tests" Requires="$(RunEditorWorkflowsGraphNodes)"/>
-
 	<!-- These Nodes run the Editor and gather profile data -->
 	<Aggregate Name="Run EditorPerf Tests" Requires="$(RunGraphNodes)"/>
 
