Fixes a handful of issues in the TSAN CitySample Tests:

- SanitizerReports only run on Warm runs (since we no longer run TSAN on Cold runs)
- We prime the TSAN Warm runs with a single Zenloader Cold run without TSAN to avoid unnecessary cold runs
- Since we require two different builds of the editor (one with TSAN and one without) we need to make the paths for those artifacts unique to avoid conflicts in the build graph. To do so, move the editor files after each compile and restore them into their expected /Binaries/<platform> path before running, using the appropriate editor binaries for Cold and Warm respectfully. Without this BuildGraph will generate hundreds of errors for each file conflict which adds excessive noise.

#rnx

diff --git a/Engine/Build/EditorPerf.xml b/Engine/Build/EditorPerf.xml
--- a/Engine/Build/EditorPerf.xml
+++ b/Engine/Build/EditorPerf.xml
@@ -99,7 +99,9 @@
 	<Property Name="AllProjects" Value=""/>
 
 	<!-- A macro that can be used to easily set up a new test project. -->
-	<Macro Name="SetupPerfProject" Arguments="ProjectName" OptionalArguments="CopyDefaultsFrom;ProjectTarget;EditorTarget;MapPath;GameSpecificArgs;UseBinariesFrom;ExtraNonCopyableArgs;GauntletTest;GauntletTestArgs">
+	<Macro Name="SetupPerfProject" Arguments="ProjectName;ProjectPath" OptionalArguments="CopyDefaultsFrom;ProjectTarget;EditorTarget;MapPath;GameSpecificArgs;UseBinariesFrom;ExtraNonCopyableArgs;GauntletTest;GauntletTestArgs">
+		<Property Name="$(ProjectName)__ProjectPath" Value="$(ProjectPath)" CreateInParentScope="true"/>
+		
 		<Property Name="$(ProjectName)__ProjectTarget" Value="$(ProjectName)" CreateInParentScope="true"/>
 		<Property Name="$(ProjectName)__ProjectTarget" Value="$($(CopyDefaultsFrom)__ProjectTarget)" If="'$(CopyDefaultsFrom)' != ''"/>
 		<Property Name="$(ProjectName)__ProjectTarget" Value="$(ProjectTarget)" If="'$(ProjectTarget)' != ''"/>
@@ -267,10 +269,10 @@
 	
 	<Property Name="GenerateTSANReportNodeName" Value=""/>
 	<Property Name="GenerateTSANReportNodeList" Value=""/>
-	<Macro Name="GenerateTSANReport" Arguments="NodeName;TestOutputDir;Phase;TestNameSuffix">
+	<Macro Name="GenerateTSANReport" Arguments="NodeName;TestOutputDir;Phase;TestNameSuffix;Requires">
 		<Property Name="LogFilePath" Value="$(TestOutputDir)/$(Phase)$(TestNameSuffix)/Editor/EditorOutput.log"/>
 		<Property Name="GenerateTSANReportNodeName" Value="Sanitizer Report - $(NodeName)"/>
-		<Node Name="$(GenerateTSANReportNodeName)">
+		<Node Name="$(GenerateTSANReportNodeName)" Requires="$(Requires)">
 			<Command Name="LogSanitizerOutput" Arguments="-logfile=&quot;$(LogFilePath)&quot;" />
 		</Node>
 		<Property Name="GenerateTSANReportNodeList" Value="$(GenerateTSANReportNodeList); $(GenerateTSANReportNodeName)"/>
@@ -298,6 +300,13 @@
 		</Do>
 	</Macro>
 
+	<Property Name="NeedToCopyTSanBinaries" Value="$(WithTSAN)"/>
+	<Macro Name="CopyTSANBinariesIfNeeded" Arguments="ProjectBinariesPath">
+		<!-- Perform swap of built artifacts for TSAN builds so we use a non-TSAN build for cold runs and a TSAN build for warm runs -->
+		<Copy From="$(ProjectBinariesPath)_TSAN/..." To="$(ProjectBinariesPath)" ErrorIfNotFound="true" Overwrite="True" If="$(NeedToCopyTSanBinaries)"/>
+		<Property Name="NeedToCopyTSanBinaries" Value="false"/>
+	</Macro>
+	
 	<!-- This will be set to the names of all nodes we want to execute based on specified platforms/options -->
 	<Property Name="BaselineGraphNodes" Value=""/>
 
@@ -316,26 +325,30 @@
 			<Property Name="OutputName" Value="$(Project)$(Platform)EditorPerf"/>
 
 			<!-- Compile the Editor. Set up compilation nodes for all projects that have been defined even if we aren't going to test them. This lets the tested projects depend on non-tested ones. -->
+			<Property Name="SharedCompileProduces" Value="#SharedBinaries_$(Project)_$(Platform)"/>
 			<Property Name="CompileNodeName" Value="Compile $(Project) $(Platform)" />
-			<Property Name="TSANCompileNodeName" Value="Compile $(Project) $(Platform) with TSAN" />
+			<Property Name="CompileNodeNameProduces" Value="#Binaries_$(Project)_$(Platform)"/>
+			<Property Name="TSANCompileNodeNameProduces" Value=""/>
+			<Property Name="TSANCompileNodeNameProduces" Value="#TsanBinaries_$(Project)_$(Platform)" If="$(WithTSAN)"/>
+			<Property Name="ProjectBinariesPath" Value="$(RootDir)/$($(Project)__ProjectPath)/Binaries/$(Platform)"/>
 			<Property Name="AgentModifier" Value=""/>
 			<Property Name="AgentModifier" Value=" $(Project)" If="!$(SingleAgentCompile)"/>
 			<Agent Name="Compile Editor Agent $(Platform)$(AgentModifier" Type="$($(Platform)CompileAgentType)">
 				<!-- 
 				If running with TSAN we compile twice, once without TSAN enabled (used for cold runs) and once with TSAN enabled (used for warm runs).
 				TSAN adds too much overhead to enable on cold runs, but we need a cold run to execute to prime the cache for the warm run so we run both.
+				We need to copy out the files to unique directories or else the buildgraph will complain about two nodes using files with the same name but
+				having different contents resulting in unneccessary errors appearing in jobs.
 				-->
-				<Node Name="$(CompileNodeName)">
+				<Node Name="$(CompileNodeName)" Produces="$(SharedCompileProduces);$(CompileNodeNameProduces);$(TSANCompileNodeNameProduces)" IgnoreModified="$(ProjectBinariesPath)/...">
 					<SetVersion Change="$(Change)" Branch="$(EscapedBranch)" If="$(IsBuildMachine)"/>
-					<Compile Target="ShaderCompileWorker" Platform="$(Platform)" Configuration="Development"/>
-					<Compile Target="CrashReportClientEditor" Platform="$(Platform)" Project="$(ProjectTarget)" Configuration="Shipping"/>
+					<Compile Target="ShaderCompileWorker" Platform="$(Platform)" Configuration="Development" Tag="$(SharedCompileProduces)"/>
+					<Compile Target="CrashReportClientEditor" Platform="$(Platform)" Project="$(ProjectTarget)" Configuration="Shipping" Tag="$(SharedCompileProduces)"/>
 					<Compile Target="$(EditorTarget)" Platform="$(Platform)" Project="$(ProjectTarget)" Configuration="Development" Arguments="$(EditorCompileArgs)"/>
-				</Node>
-				<Node Name="$(TSANCompileNodeName)" If="$(WithTSAN)">
-					<SetVersion Change="$(Change)" Branch="$(EscapedBranch)" If="$(IsBuildMachine)"/>
-					<Compile Target="ShaderCompileWorker" Platform="$(Platform)" Configuration="Development"/>
-					<Compile Target="CrashReportClientEditor" Platform="$(Platform)" Project="$(ProjectTarget)" Configuration="Shipping"/>
-					<Compile Target="$(EditorTarget)" Platform="$(Platform)" Project="$(ProjectTarget)" Configuration="Development" Arguments="$(EditorCompileArgs) $(TSANEditorCompileArgs)"/>
+					
+					<Copy From="$(ProjectBinariesPath)/..." To="$(ProjectBinariesPath)_NonTSAN" Tag="$(CompileNodeNameProduces)" ErrorIfNotFound="true" If="$(WithTSAN)"/>
+					<Compile Target="$(EditorTarget)" Platform="$(Platform)" Project="$(ProjectTarget)" Configuration="Development" Arguments="$(EditorCompileArgs) $(TSANEditorCompileArgs)" If="$(WithTSAN)"/>
+					<Copy From="$(ProjectBinariesPath)/..." To="$(ProjectBinariesPath)_TSAN" Tag="$(TSANCompileNodeNameProduces)" ErrorIfNotFound="true" If="$(WithTSAN)"/>
 				</Node>
 			</Agent>
 			<Label Category="$(CompileCategory)" Name="$(Project) Editor Compile" Requires="$(CompileNodeName)"/>
@@ -348,82 +361,107 @@
 				<Property Name="RunNodeList" Value=""/>		
 				<Property Name="UploadNodeList" Value=""/>
 				<Agent Name="EditorPerf Agent $(Platform) $(Project)" Type="EditorPerf$(Platform)">
+					<Property Name="WarmTSANNodeRequires" Value="" />
+					
 					<ForEach Name="Phase" Values="$(Phases)">
-
 						<Property Name="RunNodeCompilationRequires" Value="$(CompileNodeName)"/>
 						<Property Name="RunNodeCompilationRequires" Value="$($($(Project)__UseBinariesFrom)__CompileNodeName)" If="'$($(Project)__UseBinariesFrom)'!=''"/>
-						<!-- Only use a TSAN enabled builds on warm builds (if using TSAN at all) -->
-						<Property Name="RunNodeCompilationRequires" Value="$(TSANCompileNodeName)" If="$(WithTSAN) and '$(Phase)' == 'Warm'"/>
-
-						<Do If="$(EnableBasePhaseTests)">
-							<ForEach Name="Iteration" Values="$(Iterations)">
-								<Property Name="IterationSuffix" Value=""/>
-								<Property Name="IterationSuffix" Value="$(Iteration)" If="'$(Iteration)' != '1'"/>
-								<Property Name="RunNodeName" Value="Run $(OutputName) $(Phase)$(IterationSuffix)"/>
-								<Property Name="RunNodeRequires" Value="$(RunNodeCompilationRequires)"/>
-								<Property Name="RunNodeRequires" Value="$(RunNodeCompilationRequires); Run $(OutputName) Cold$(IterationSuffix)" If="'$(Phase)' == 'Warm'"/>
-								<Node Name="$(RunNodeName)" Requires="$(RunNodeRequires)">
-									<Expand Name="RunTest" TestNameSuffix="NoZenLoader" Platform="$(Platform)" ProjectName="$(Project)" ProjectTarget="$(ProjectTarget)" NodeName="$(RunNodeName)" Phase="$(Phase)" Iteration="$(Iteration)" Args="$($(Phase)Args)" EditorArgs="$(ProjectCombinedArgs) $(NoZenLoaderTestArgs)"/>
-								</Node>
-								<Property Name="RunNodeList" Value="$(RunNodeList);$(RunNodeName)"/>
-								<Expand Name="GenerateTSANReport" NodeName="$(RunNodeName)" TestOutputDir="$(TestOutputDir)" Phase="$(Phase)" TestNameSuffix="NoZenLoader" If="$(WithTSAN)"/>
-							</ForEach>
-						</Do>
+						<Property Name="IsWarmTSAN" Value="(('$(Phase)' == 'Warm') and $(WithTSAN))"/>
+						<Property Name="IsColdTSAN" Value="(('$(Phase)' == 'Cold') and $(WithTSAN))"/>
 						
-						<!-- Setup Zen Loader tests for this phase -->
-						<Do>
-							<!--Zen Loader TestRun if enabled--> 
-							<Property Name="RunNodeZenName" Value="Run $(OutputName) $(Phase) ZenLoader"/>
-							<Property Name="RunNodeZenRequires" Value="$(RunNodeCompilationRequires)"/>
-							<Property Name="RunNodeZenRequires" Value="$(RunNodeCompilationRequires); Run $(OutputName) Cold ZenLoader" If="'$(Phase)' == 'Warm'"/>
-							<Node Name="$(RunNodeZenName)" Requires="$(RunNodeZenRequires)" If="$(EnableZenLoaderTest)">
+						<!-- Cold TSAN case -->
+						<Do If="$(IsColdTSAN)">
+							<!-- Force a Cold Zen Loader run to prime the DDC for our actual test case (whatever it is, zen, legacy, lazy loader etc...) -->
+							<Property Name="RunNodeZenName" Value="Run Non-TSAN $(OutputName) $(Phase) ZenLoader"/>
+							<Node Name="$(RunNodeZenName)" Requires="$(CompileNodeName);$(SharedCompileProduces);$(CompileNodeNameProduces)" IgnoreModified="$(ProjectBinariesPath)/...">
+								<!-- Perform swap of built artifacts for TSAN builds so we use a non-TSAN build for cold runs and a TSAN build for warm runs -->
+								<Copy From="$(ProjectBinariesPath)_NonTSAN/..." To="$(ProjectBinariesPath)" ErrorIfNotFound="true" Overwrite="True"/>
 								<Expand Name="RunTest" TestNameSuffix="ZenLoader" Platform="$(Platform)" ProjectName="$(Project)" ProjectTarget="$(ProjectTarget)" NodeName="$(RunNodeZenName)" Phase="$(Phase)" Iteration="1" Args="$($(Phase)Args)" EditorArgs="$(ProjectCombinedArgs) $(ZenLoaderTestArgs)"/>
 							</Node>
-							<Property Name="RunNodeList" Value="$(RunNodeList);$(RunNodeZenName)" If="$(EnableZenLoaderTest)"/>
-							<Expand Name="GenerateTSANReport" NodeName="$(RunNodeZenName)" TestOutputDir="$(TestOutputDir)" Phase="$(Phase)" TestNameSuffix="ZenLoader" If="$(EnableZenLoaderTest) And $(WithTSAN)"/>
-							
-							<!--Zen Loader Async TestRun if enabled--> 
-							<Property Name="RunNodeZenAsyncName" Value="Run $(OutputName) $(Phase) ZenLoader Async"/>
-							<Property Name="RunNodeZenAsyncRequires" Value="$(RunNodeCompilationRequires)"/>
-							<Property Name="RunNodeZenAsyncRequires" Value="$(RunNodeCompilationRequires); Run $(OutputName) Cold ZenLoader Async" If="'$(Phase)' == 'Warm'"/>
-							<Node Name="$(RunNodeZenAsyncName)" Requires="$(RunNodeZenAsyncRequires)" If="$(EnableZenLoaderAsyncThreadTest)">
-								<Expand Name="RunTest" TestNameSuffix="ZenLoaderAsync" Platform="$(Platform)" ProjectName="$(Project)" ProjectTarget="$(ProjectTarget)" NodeName="$(RunNodeZenAsyncName)" Phase="$(Phase)" Iteration="1" Args="$($(Phase)Args)" EditorArgs="$(ProjectCombinedArgs) $(ZenLoaderAsyncTestArgs)"/>
-							</Node>
-							<Property Name="RunNodeList" Value="$(RunNodeList);$(RunNodeZenAsyncName)" If="$(EnableZenLoaderAsyncThreadTest)"/>
-							<Expand Name="GenerateTSANReport" NodeName="$(RunNodeZenAsyncName)" TestOutputDir="$(TestOutputDir)" Phase="$(Phase)" TestNameSuffix="ZenLoaderAsync" If="$(EnableZenLoaderAsyncThreadTest) And $(WithTSAN)"/>
-
-							<!--Zen Store DDC TestRun if enabled--> 
-							<Property Name="RunNodeZenStoreDDCName" Value="Run $(OutputName) $(Phase) ZenLoader Store DDC"/>
-							<Property Name="RunNodeZenStoreDDCRequires" Value="$(RunNodeCompilationRequires)"/>
-							<Property Name="RunNodeZenStoreDDCRequires" Value="$(RunNodeCompilationRequires); Run $(OutputName) Cold ZenLoader Store DDC" If="'$(Phase)' == 'Warm'"/>
-							<Node Name="$(RunNodeZenStoreDDCName)" Requires="$(RunNodeZenStoreDDCRequires)" If="$(EnableZenStoreDDCTest)">
-								<Expand Name="RunTest" TestNameSuffix="ZenStoreAsyncDDC" Platform="$(Platform)" ProjectName="$(Project)" ProjectTarget="$(ProjectTarget)" NodeName="$(RunNodeZenStoreDDCName)" Phase="$(Phase)" Iteration="1" Args="$(ZenStoreDDCPhaseArgs)" EditorArgs="$(ProjectCombinedArgs) $(ZenStoreDDCTestArgs)"/>
-							</Node>
-							<Property Name="RunNodeList" Value="$(RunNodeList);$(RunNodeZenStoreDDCName)" If="$(EnableZenStoreDDCTest)"/>
-							<Expand Name="GenerateTSANReport" NodeName="$(RunNodeZenStoreDDCName)" TestOutputDir="$(TestOutputDir)" Phase="$(Phase)" TestNameSuffix="ZenStoreAsyncDDC" If="$(EnableZenStoreDDCTest) And $(WithTSAN)"/>
+							<Property Name="RunNodeList" Value="$(RunNodeList);$(RunNodeZenName)"/>
+							<Property Name="WarmTSANNodeRequires" Value="$(RunNodeZenName)"/>
 						</Do>
 						
-						<!--Warm only test runs e.g LazyLoad--> 
-						<Do If="'$(Phase)' == 'Warm'">
-							<!--Legacy LazyLoad TestRun if enabled--> 
-							<Property Name="RunNodeLazyLoadName" Value="Run $(OutputName) $(Phase) LazyLoad NoZenLoader"/>
-							<Property Name="RunNodeLazyLoadRequires" Value="$(RunNodeCompilationRequires)"/>
-							<Property Name="RunNodeLazyLoadRequires" Value="$(RunNodeCompilationRequires); Run $(OutputName) Cold LazyLoad NoZenLoader" If="'$(Phase)' == 'Warm'"/>
-							<Node Name="$(RunNodeLazyLoadName)" Requires="$(RunNodeLazyLoadRequires)" If="$(EnableLazyLoadNoZenLoaderTest)">
-								<Expand Name="RunTest" TestNameSuffix="NoZenLoader" Platform="$(Platform)" ProjectName="$(Project)" ProjectTarget="$(ProjectTarget)" NodeName="$(RunNodeLazyLoadName)" Phase="$(Phase)" Iteration="1" Args="$($(Phase)Args)" EditorArgs="$(ProjectCombinedArgs) $(LazyLoadTestArgs) $(NoZenLoaderTestArgs)"/>
-							</Node>
-							<Property Name="RunNodeList" Value="$(RunNodeList);$(RunNodeLazyLoadName)" If="$(EnableLazyLoadNoZenLoaderTest)"/>
-							<Expand Name="GenerateTSANReport" NodeName="$(RunNodeLazyLoadName)" TestOutputDir="$(TestOutputDir)" Phase="$(Phase)" TestNameSuffix="NoZenLoader" If="$(EnableLazyLoadNoZenLoaderTest) And $(WithTSAN)"/>
+						<!-- Cold and Warm Non-TSAN case, and Warm TSAN cases -->
+						<Do If="!$(IsColdTSAN)">
+							<Property Name="RunNodeCompilationRequires" Value="$(WarmTSANNodeRequires);$(SharedCompileProduces);$(TSANCompileNodeNameProduces)" If="$(IsWarmTSAN)"/>
 							
-							<!--Zen LazyLoad TestRun if enabled--> 
-							<Property Name="RunNodeLazyLoadZenName" Value="Run $(OutputName) $(Phase) LazyLoad ZenLoader"/>
-							<Property Name="RunNodeLazyLoadZenRequires" Value="$(RunNodeCompilationRequires)"/>
-							<Property Name="RunNodeLazyLoadZenRequires" Value="$(RunNodeCompilationRequires); Run $(OutputName) Cold LazyLoad ZenLoader" If="'$(Phase)' == 'Warm'"/>
-							<Node Name="$(RunNodeLazyLoadZenName)" Requires="$(RunNodeLazyLoadZenRequires);" If="$(EnableLazyLoadZenLoaderTest)">
-								<Expand Name="RunTest" TestNameSuffix="ZenLoader" Platform="$(Platform)" ProjectName="$(Project)" ProjectTarget="$(ProjectTarget)" NodeName="$(RunNodeLazyLoadZenName)" Phase="$(Phase)" Iteration="1" Args="$($(Phase)Args)" EditorArgs="$(ProjectCombinedArgs) $(LazyLoadZenTestArgs)"/>
-							</Node>
-							<Property Name="RunNodeList" Value="$(RunNodeList);$(RunNodeLazyLoadZenName)" If="$(EnableLazyLoadZenLoaderTest)"/>
-							<Expand Name="GenerateTSANReport" NodeName="$(RunNodeLazyLoadZenName)" TestOutputDir="$(TestOutputDir)" Phase="$(Phase)" TestNameSuffix="ZenLoader" If="$(EnableLazyLoadZenLoaderTest) And $(WithTSAN)"/>
+							<Do If="$(EnableBasePhaseTests)">
+								<ForEach Name="Iteration" Values="$(Iterations)">
+									<Property Name="IterationSuffix" Value=""/>
+									<Property Name="IterationSuffix" Value="$(Iteration)" If="'$(Iteration)' != '1'"/>
+									<Property Name="RunNodeName" Value="Run $(OutputName) $(Phase)$(IterationSuffix)"/>
+									<Property Name="RunNodeRequires" Value="$(RunNodeCompilationRequires)"/>
+									<Property Name="RunNodeRequires" Value="$(RunNodeCompilationRequires); Run $(OutputName) Cold$(IterationSuffix)" If="('$(Phase)' == 'Warm') and !$(WithTSAN)"/>
+									<Node Name="$(RunNodeName)" Requires="$(RunNodeRequires)" IgnoreModified="$(ProjectBinariesPath)/...">
+										<Expand Name="CopyTSANBinariesIfNeeded" ProjectBinariesPath="$(ProjectBinariesPath)"/>
+										<Expand Name="RunTest" TestNameSuffix="NoZenLoader" Platform="$(Platform)" ProjectName="$(Project)" ProjectTarget="$(ProjectTarget)" NodeName="$(RunNodeName)" Phase="$(Phase)" Iteration="$(Iteration)" Args="$($(Phase)Args)" EditorArgs="$(ProjectCombinedArgs) $(NoZenLoaderTestArgs)"/>
+									</Node>
+									<Property Name="RunNodeList" Value="$(RunNodeList);$(RunNodeName)"/>
+									<Expand Name="GenerateTSANReport" NodeName="$(RunNodeName)" TestOutputDir="$(TestOutputDir)" Phase="$(Phase)" TestNameSuffix="NoZenLoader" Requires="$(WarmTSANNodeRequires)" If="$(IsWarmTSAN)"/>
+								</ForEach>
+							</Do>
+						
+							<!-- Setup Zen Loader tests for this phase -->
+							<Do>
+								<!--Zen Loader TestRun if enabled--> 
+								<Property Name="RunNodeZenName" Value="Run $(OutputName) $(Phase) ZenLoader"/>
+								<Property Name="RunNodeZenRequires" Value="$(RunNodeCompilationRequires)"/>
+								<Property Name="RunNodeZenRequires" Value="$(RunNodeCompilationRequires); Run $(OutputName) Cold ZenLoader" If="('$(Phase)' == 'Warm') and !$(WithTSAN)"/>
+								<Node Name="$(RunNodeZenName)" Requires="$(RunNodeZenRequires)" IgnoreModified="$(ProjectBinariesPath)/..." If="$(EnableZenLoaderTest)">
+									<Expand Name="CopyTSANBinariesIfNeeded" ProjectBinariesPath="$(ProjectBinariesPath)"/>
+									<Expand Name="RunTest" TestNameSuffix="ZenLoader" Platform="$(Platform)" ProjectName="$(Project)" ProjectTarget="$(ProjectTarget)" NodeName="$(RunNodeZenName)" Phase="$(Phase)" Iteration="1" Args="$($(Phase)Args)" EditorArgs="$(ProjectCombinedArgs) $(ZenLoaderTestArgs)"/>
+								</Node>
+								<Property Name="RunNodeList" Value="$(RunNodeList);$(RunNodeZenName)" If="$(EnableZenLoaderTest)"/>
+								<Expand Name="GenerateTSANReport" NodeName="$(RunNodeZenName)" TestOutputDir="$(TestOutputDir)" Phase="$(Phase)" TestNameSuffix="ZenLoader" Requires="$(WarmTSANNodeRequires)" If="$(EnableZenLoaderTest) And $(IsWarmTSAN)"/>
+							
+								<!--Zen Loader Async TestRun if enabled--> 
+								<Property Name="RunNodeZenAsyncName" Value="Run $(OutputName) $(Phase) ZenLoader Async"/>
+								<Property Name="RunNodeZenAsyncRequires" Value="$(RunNodeCompilationRequires)"/>
+								<Property Name="RunNodeZenAsyncRequires" Value="$(RunNodeCompilationRequires); Run $(OutputName) Cold ZenLoader Async" If="('$(Phase)' == 'Warm') and !$(WithTSAN)"/>
+								<Node Name="$(RunNodeZenAsyncName)" Requires="$(RunNodeZenAsyncRequires)" IgnoreModified="$(ProjectBinariesPath)/..." If="$(EnableZenLoaderAsyncThreadTest)">
+									<Expand Name="CopyTSANBinariesIfNeeded" ProjectBinariesPath="$(ProjectBinariesPath)"/>
+									<Expand Name="RunTest" TestNameSuffix="ZenLoaderAsync" Platform="$(Platform)" ProjectName="$(Project)" ProjectTarget="$(ProjectTarget)" NodeName="$(RunNodeZenAsyncName)" Phase="$(Phase)" Iteration="1" Args="$($(Phase)Args)" EditorArgs="$(ProjectCombinedArgs) $(ZenLoaderAsyncTestArgs)"/>
+								</Node>
+								<Property Name="RunNodeList" Value="$(RunNodeList);$(RunNodeZenAsyncName)" If="$(EnableZenLoaderAsyncThreadTest)"/>
+								<Expand Name="GenerateTSANReport" NodeName="$(RunNodeZenAsyncName)" TestOutputDir="$(TestOutputDir)" Phase="$(Phase)" TestNameSuffix="ZenLoaderAsync" Requires="$(WarmTSANNodeRequires)" If="$(EnableZenLoaderAsyncThreadTest) And $(IsWarmTSAN)"/>
+
+								<!--Zen Store DDC TestRun if enabled--> 
+								<Property Name="RunNodeZenStoreDDCName" Value="Run $(OutputName) $(Phase) ZenLoader Store DDC"/>
+								<Property Name="RunNodeZenStoreDDCRequires" Value="$(RunNodeCompilationRequires)"/>
+								<Property Name="RunNodeZenStoreDDCRequires" Value="$(RunNodeCompilationRequires); Run $(OutputName) Cold ZenLoader Store DDC" If="('$(Phase)' == 'Warm') and !$(WithTSAN)"/>
+								<Node Name="$(RunNodeZenStoreDDCName)" Requires="$(RunNodeZenStoreDDCRequires)" IgnoreModified="$(ProjectBinariesPath)/..." If="$(EnableZenStoreDDCTest)">
+									<Expand Name="CopyTSANBinariesIfNeeded" ProjectBinariesPath="$(ProjectBinariesPath)"/>
+									<Expand Name="RunTest" TestNameSuffix="ZenStoreAsyncDDC" Platform="$(Platform)" ProjectName="$(Project)" ProjectTarget="$(ProjectTarget)" NodeName="$(RunNodeZenStoreDDCName)" Phase="$(Phase)" Iteration="1" Args="$(ZenStoreDDCPhaseArgs)" EditorArgs="$(ProjectCombinedArgs) $(ZenStoreDDCTestArgs)"/>
+								</Node>
+								<Property Name="RunNodeList" Value="$(RunNodeList);$(RunNodeZenStoreDDCName)" If="$(EnableZenStoreDDCTest)"/>
+								<Expand Name="GenerateTSANReport" NodeName="$(RunNodeZenStoreDDCName)" TestOutputDir="$(TestOutputDir)" Phase="$(Phase)" TestNameSuffix="ZenStoreAsyncDDC" Requires="$(WarmTSANNodeRequires)" If="$(EnableZenStoreDDCTest) And $(IsWarmTSAN)"/>
+							</Do>
+						
+							<!--Warm only test runs e.g LazyLoad--> 
+							<Do If="'$(Phase)' == 'Warm'">
+								<!--Legacy LazyLoad TestRun if enabled--> 
+								<Property Name="RunNodeLazyLoadName" Value="Run $(OutputName) $(Phase) LazyLoad NoZenLoader"/>
+								<Property Name="RunNodeLazyLoadRequires" Value="$(RunNodeCompilationRequires)"/>
+								<Property Name="RunNodeLazyLoadRequires" Value="$(RunNodeCompilationRequires); Run $(OutputName) Cold LazyLoad NoZenLoader" If="('$(Phase)' == 'Warm') and !$(WithTSAN)"/>
+								<Node Name="$(RunNodeLazyLoadName)" Requires="$(RunNodeLazyLoadRequires)" IgnoreModified="$(ProjectBinariesPath)/..." If="$(EnableLazyLoadNoZenLoaderTest)">
+									<Expand Name="CopyTSANBinariesIfNeeded" ProjectBinariesPath="$(ProjectBinariesPath)"/>
+									<Expand Name="RunTest" TestNameSuffix="NoZenLoader" Platform="$(Platform)" ProjectName="$(Project)" ProjectTarget="$(ProjectTarget)" NodeName="$(RunNodeLazyLoadName)" Phase="$(Phase)" Iteration="1" Args="$($(Phase)Args)" EditorArgs="$(ProjectCombinedArgs) $(LazyLoadTestArgs) $(NoZenLoaderTestArgs)"/>
+								</Node>
+								<Property Name="RunNodeList" Value="$(RunNodeList);$(RunNodeLazyLoadName)" If="$(EnableLazyLoadNoZenLoaderTest)"/>
+								<Expand Name="GenerateTSANReport" NodeName="$(RunNodeLazyLoadName)" TestOutputDir="$(TestOutputDir)" Phase="$(Phase)" TestNameSuffix="NoZenLoader" Requires="$(WarmTSANNodeRequires)" If="$(EnableLazyLoadNoZenLoaderTest) And $(IsWarmTSAN)"/>
+							
+								<!--Zen LazyLoad TestRun if enabled--> 
+								<Property Name="RunNodeLazyLoadZenName" Value="Run $(OutputName) $(Phase) LazyLoad ZenLoader"/>
+								<Property Name="RunNodeLazyLoadZenRequires" Value="$(RunNodeCompilationRequires)"/>
+								<Property Name="RunNodeLazyLoadZenRequires" Value="$(RunNodeCompilationRequires); Run $(OutputName) Cold LazyLoad ZenLoader" If="('$(Phase)' == 'Warm') and !$(WithTSAN)"/>
+								<Node Name="$(RunNodeLazyLoadZenName)" Requires="$(RunNodeLazyLoadZenRequires);" IgnoreModified="$(ProjectBinariesPath)/..." If="$(EnableLazyLoadZenLoaderTest)">
+									<Expand Name="CopyTSANBinariesIfNeeded" ProjectBinariesPath="$(ProjectBinariesPath)"/>
+									<Expand Name="RunTest" TestNameSuffix="ZenLoader" Platform="$(Platform)" ProjectName="$(Project)" ProjectTarget="$(ProjectTarget)" NodeName="$(RunNodeLazyLoadZenName)" Phase="$(Phase)" Iteration="1" Args="$($(Phase)Args)" EditorArgs="$(ProjectCombinedArgs) $(LazyLoadZenTestArgs)"/>
+								</Node>
+								<Property Name="RunNodeList" Value="$(RunNodeList);$(RunNodeLazyLoadZenName)" If="$(EnableLazyLoadZenLoaderTest)"/>
+								<Expand Name="GenerateTSANReport" NodeName="$(RunNodeLazyLoadZenName)" TestOutputDir="$(TestOutputDir)" Phase="$(Phase)" TestNameSuffix="ZenLoader" Requires="$(WarmTSANNodeRequires)" If="$(EnableLazyLoadZenLoaderTest) And $(IsWarmTSAN)"/>
+							</Do>
 						</Do>
 					</ForEach>
 
