Fix warnings in LLT jobs
#rnx

diff --git a/Engine/Build/LowLevelTests.xml b/Engine/Build/LowLevelTests.xml
--- a/Engine/Build/LowLevelTests.xml
+++ b/Engine/Build/LowLevelTests.xml
@@ -73,13 +73,8 @@
 
 	<!-- Test Options -->
 	<Option Name="RunAllTests" DefaultValue="true" Description="Run All Available Tests (except Foundation)"/>
+	<Option Name="DryRun" DefaultValue="false" Description="Dry Run"/>
 	
-	<!-- If we set one of the Run*Tests option we don't want to run all the tests, just the selected ones. -->
-	<Property Name="RunsAllTests" Value="$(RunAllTests)"/>
-	<ForEach Name="TestName" Values="$(TestNames)">
-		<Property Name="RunsAllTests" Value="false" If="$(Run$(TestName)Tests)"/>
-	</ForEach>
-
 	<!-- Platform Options -->
 	<Option Name="TestAllSupportedPlatforms" DefaultValue="true" Description="Test All Supported Platforms"/>
 	<Option Name="TestPlatformWin64" DefaultValue="false" Description="Test Platform Win64"/>
@@ -127,19 +122,23 @@
 
 	<Macro Name="TestAfterStep" Arguments="Platform;TestName;ShortName">
 		<Node Name="$(TestName) AfterStep $(Platform)">
-			<Expand Name="$(TestName)AfterSteps" Platform="$(Platform)" />	
+			<Expand Name="$(TestName)AfterSteps" Platform="$(Platform)" If="!$(DryRun)"/>
 		</Node>
 		<Label Category="After" Name="$(ShortName) After $(Platform)" Requires="$(TestName) AfterStep $(Platform)"/>
-		<Property Name="AfterTestNodes" Value="$(AfterTestNodes);$(TestName) AfterStep $(Platform)" If="($(RunsAllTests) Or $(Run$(TestName)Tests)) And ($(TestsAllSupportedPlatforms) Or $(TestPlatform$(Platform)))"/>
+		<Property Name="TestRuns" Value="($(RunsAllTests) and ('$(Run$(TestName)Tests)' != 'false')) or (!$(RunsAllTests) and ('$(Run$(TestName)Tests)' == 'true'))"/>
+		<Property Name="TestRuns" Value="$(TestRuns) and ($(TestsAllSupportedPlatforms) or $(TestPlatform$(Platform)))"/>
+		<Do If="$(TestRuns)">
+			<Property Name="AfterTestNodes" Value="$(AfterTestNodes);$(TestName) AfterStep $(Platform)"/>
+		</Do>
 	</Macro>
 	
 	<Macro Name="BuildAndStageTest" Arguments="Platform;TestName;StagesWithProjectFile;TargetName;ExtraCompilationArgs">
 		<Property Name="LocalAchiveDirectory" Value="$(RootDir)/$(TargetName)"/>
 		<Do If="'$(StagesWithProjectFile)' == 'True'">
-			<Command Name="BuildCookRun" Arguments="-Build -NoBootstrapExe -NoCodeSign -SkipCook -Stage -NoSubmit -Archive -ArchiveDirectory=$(LocalAchiveDirectory) -project=$(TargetName) -configuration=$(Configuration) -platform=$(Platform) $(ExtraCompilationArgs)"/>
+			<Command Name="BuildCookRun" Arguments="-Build -NoBootstrapExe -NoCodeSign -SkipCook -Stage -NoSubmit -Archive -ArchiveDirectory=$(LocalAchiveDirectory) -project=$(TargetName) -configuration=$(Configuration) -platform=$(Platform) $(ExtraCompilationArgs)" If="!$(DryRun)"/>
 		</Do>
 		<Do If="'$(StagesWithProjectFile)' != 'True'">
-			<Compile Target="$(TargetName)" Configuration="$(Configuration)" Platform="$(Platform)" AllowXGE="false" Arguments="-Mode=Test $(ExtraCompilationArgs)"/>
+			<Compile Target="$(TargetName)" Configuration="$(Configuration)" Platform="$(Platform)" AllowXGE="false" Arguments="-Mode=Test $(ExtraCompilationArgs)" If="!$(DryRun)"/>
 		</Do>
 	</Macro>
 	
@@ -153,6 +152,9 @@
 		<Property Name="BuildPlatform" Value=""/>
 
 		<Expand Name="SelectBuildPlatformAndAgentType" Platform="$(Platform)"/>
+
+		<Property Name="TestRuns" Value="($(RunsAllTests) and ('$(Run$(TestName)Tests)' != 'false')) or (!$(RunsAllTests) and ('$(Run$(TestName)Tests)' == 'true'))"/>
+		<Property Name="TestRuns" Value="$(TestRuns) and ($(TestsAllSupportedPlatforms) or $(TestPlatform$(Platform)))"/>
 		
 		<Agent Name="Low Level Tests $(Platform)" Type="$(AgentType)">
 			<Property Name="TestRequirements" Value=""/>
@@ -160,10 +162,10 @@
 				<Property Name="TestRequirements" Value="$(TestRequirements);Catch2 Build Library $(Platform)"/>
 			</Do>
 			<Node Name="$(ShortName) Tests $(Platform)" Requires="$(TestRequirements)">
-				<Do If="'$(Deactivated)' == ''">
+				<Do If="'$(Deactivated)' == '' and $(TestRuns)">
 					<Expand Name="BuildAndStageTest" TestName="$(TestName)" Platform="$(Platform)" TargetName="$(TargetName)" StagesWithProjectFile="$(StagesWithProjectFile)" ExtraCompilationArgs="$(ExtraCompilationArgs)"/>
 				</Do>
-				<Do If="'$(RunUnsupported)' == '' AND '$(Deactivated)' == ''">
+				<Do If="'$(RunUnsupported)' == '' and '$(Deactivated)' == '' and $(TestRuns)">
 					<Property Name="PlatformStageDirName" Value="$(Platform)"/>
 					<Property Name="PlatformStageDirName" Value="Windows" If="'$(Platform)' == 'Win64'"/>
 					<Property Name="DevicePool" Value="$(GauntletDevicePool)" />
@@ -203,8 +205,8 @@
 						<StringOp Input="$(BinaryRelativePath)" Method="Replace" Arguments="\;/" Output="BinariesRelativeForwardSlash"/>
 
 						<!-- Builds docker image -->
-						<Expand Name="$(Platform)DockerCleanup" AppNameLowercase="$(TestTargetLowercase)"/>
-						<Expand Name="$(Platform)DockerSetup" AppNameLowercase="$(TestTargetLowercase)" AppDir="$(BinariesRelativeForwardSlash)/$(Platform)/$(TargetName)$(TestModeSuffix)"/>
+						<Expand Name="$(Platform)DockerCleanup" AppNameLowercase="$(TestTargetLowercase)" If="!$(DryRun)"/>
+						<Expand Name="$(Platform)DockerSetup" AppNameLowercase="$(TestTargetLowercase)" AppDir="$(BinariesRelativeForwardSlash)/$(Platform)/$(TargetName)$(TestModeSuffix)" If="!$(DryRun)"/>
 						<Property Name="TestArgsExt" Value="$(TestArgsExt) -containerized"/>
 					</Do>
 
@@ -214,20 +216,22 @@
 					</Do>
 
 					<Command Name="RunLowLevelTests"
-							Arguments="-test=LowLevelTests -testapp=$(TargetName)$(TestModeSuffix) -platform=$(Platform) -configuration=$(Configuration) $(TestArgsExt) -VeryVerbose"/>
+							Arguments="-test=LowLevelTests -testapp=$(TargetName)$(TestModeSuffix) -platform=$(Platform) -configuration=$(Configuration) $(TestArgsExt) -VeryVerbose" If="!$(DryRun)"/>
 
 					<Do If="'$(AppRunsContainerized)' == 'True'">
 						<Property Name="TestTargetLowercase" Value=""/>
 						<StringOp Input="$(TargetName)" Method="ToLower" Output="TestTargetLowercase"/>
 
 						<!-- Stop any running docker container and remove image -->
-						<Expand Name="$(Platform)DockerCleanup" AppNameLowercase="$(TestTargetLowercase)"/>
+						<Expand Name="$(Platform)DockerCleanup" AppNameLowercase="$(TestTargetLowercase)" If="!$(DryRun)"/>
 					</Do>
 				</Do>
 			</Node>
 			<Label Category="Test" Name="$(ShortName) $(Platform)" Requires="$(ShortName) Tests $(Platform)"/>
-			<Property Name="TestNodes" Value="$(TestNodes);$(ShortName) Tests $(Platform)" If="($(RunsAllTests) Or $(Run$(TestName)Tests)) And ($(TestsAllSupportedPlatforms) Or $(TestPlatform$(Platform)))"/>
-				
+			<Trace Message="$(ShortName) runs because $(TestRuns)" If="$(TestRuns) and $(DryRun)"/>
+			<Do If="$(TestRuns)">
+				<Property Name="TestNodes" Value="$(TestNodes);$(ShortName) Tests $(Platform)"/>
+			</Do>
 			<Do If="'$(CheckedRunAfterStep)' == 'True' AND '$(HasAfterSteps)' == 'True'">
 				<Expand Name="TestAfterStep" TestName="$(TestName)" ShortName="$(ShortName)" Platform="$(Platform)"/>
 			</Do>
@@ -285,6 +289,8 @@
 	<Include Script="../Restricted/NotForLicensees/Build/LowLevelTests/*.xml"/>
 	<Include Script="../Restricted/NotForLicensees/Platforms/*/Build/LowLevelTests/*.xml"/>
 
+	<!-- If we set one of the Run*Tests option we don't want to run all the tests, just the selected ones. -->
+	<Property Name="RunsAllTests" Value="$(RunAllTests)"/>
 	<Expand Name="RunAllTests"/>
 	
 	<Aggregate Name="Run Low Level Tests" Requires="$(TestNodes);$(AfterTestNodes)"/>
diff --git a/Engine/Build/LowLevelTests/ChasoUserDataPT.xml b/Engine/Build/LowLevelTests/ChasoUserDataPT.xml
--- a/Engine/Build/LowLevelTests/ChasoUserDataPT.xml
+++ b/Engine/Build/LowLevelTests/ChasoUserDataPT.xml
@@ -1,6 +1,6 @@
 ﻿<?xml version="1.0" encoding="utf-8"?>
 <BuildGraph xmlns:xsi="http://www.w3.org/2001/XMLSchema-instance" xsi:schemaLocation="http://www.epicgames.com/BuildGraph ../../Build/Graph/Schema.xsd" xmlns="http://www.epicgames.com/BuildGraph">
-  <Option Name="RunChasoUserDataPTTests" DefaultValue="False" Description="Run Chaos User Data PT Tests" />
+  <Option Name="RunChasoUserDataPTTests" DefaultValue="" Description="Run Chaos User Data PT Tests" />
   <Property Name="TestNames" Value="$(TestNames);ChasoUserDataPT" />
   <Extend Name="RunAllTests">
     <Expand Name="DeployAndTest" Platform="Win64" TestName="ChasoUserDataPT" ShortName="Chaos User Data PT" TargetName="ChaosUserDataPTTests" BinaryRelativePath="Engine\Binaries" ReportType="console" />
diff --git a/Engine/Build/LowLevelTests/EventLoop.xml b/Engine/Build/LowLevelTests/EventLoop.xml
--- a/Engine/Build/LowLevelTests/EventLoop.xml
+++ b/Engine/Build/LowLevelTests/EventLoop.xml
@@ -1,6 +1,6 @@
 ﻿<?xml version="1.0" encoding="utf-8"?>
 <BuildGraph xmlns:xsi="http://www.w3.org/2001/XMLSchema-instance" xsi:schemaLocation="http://www.epicgames.com/BuildGraph ../../Build/Graph/Schema.xsd" xmlns="http://www.epicgames.com/BuildGraph">
-  <Option Name="RunEventLoopTests" DefaultValue="False" Description="Run EventLoop Tests" />
+  <Option Name="RunEventLoopTests" DefaultValue="" Description="Run EventLoop Tests" />
   <Property Name="TestNames" Value="$(TestNames);EventLoop" />
   <Extend Name="RunAllTests">
     <Expand Name="DeployAndTest" Platform="Win64" TestName="EventLoop" ShortName="EventLoop" TargetName="EventLoopUnitTests" BinaryRelativePath="Engine\Binaries" ReportType="console" />
diff --git a/Engine/Build/LowLevelTests/Foundation.xml b/Engine/Build/LowLevelTests/Foundation.xml
--- a/Engine/Build/LowLevelTests/Foundation.xml
+++ b/Engine/Build/LowLevelTests/Foundation.xml
@@ -1,6 +1,6 @@
 ﻿<?xml version="1.0" encoding="utf-8"?>
 <BuildGraph xmlns:xsi="http://www.w3.org/2001/XMLSchema-instance" xsi:schemaLocation="http://www.epicgames.com/BuildGraph ../../Build/Graph/Schema.xsd" xmlns="http://www.epicgames.com/BuildGraph">
-  <Option Name="RunFoundationTests" DefaultValue="False" Description="Run Foundation Tests" />
+  <Option Name="RunFoundationTests" DefaultValue="" Description="Run Foundation Tests" />
   <Property Name="TestNames" Value="$(TestNames);Foundation" />
   <Extend Name="RunAllTests">
     <Expand Name="DeployAndTest" Platform="Win64" TestName="Foundation" ShortName="Foundation" TargetName="FoundationTests" BinaryRelativePath="Engine\Binaries" ReportType="xml" ExtraCompilationArgs="-allmodules" />
diff --git a/Engine/Build/LowLevelTests/GameplayGraph.xml b/Engine/Build/LowLevelTests/GameplayGraph.xml
--- a/Engine/Build/LowLevelTests/GameplayGraph.xml
+++ b/Engine/Build/LowLevelTests/GameplayGraph.xml
@@ -1,6 +1,6 @@
 ﻿<?xml version="1.0" encoding="utf-8"?>
 <BuildGraph xmlns:xsi="http://www.w3.org/2001/XMLSchema-instance" xsi:schemaLocation="http://www.epicgames.com/BuildGraph ../../Build/Graph/Schema.xsd" xmlns="http://www.epicgames.com/BuildGraph">
-  <Option Name="RunGameplayGraphTests" DefaultValue="False" Description="Run GameplayGraph Tests" />
+  <Option Name="RunGameplayGraphTests" DefaultValue="" Description="Run GameplayGraph Tests" />
   <Property Name="TestNames" Value="$(TestNames);GameplayGraph" />
   <Extend Name="RunAllTests">
     <Expand Name="DeployAndTest" Platform="Win64" TestName="GameplayGraph" ShortName="GameplayGraph" TargetName="GameplayGraphTests" BinaryRelativePath="Engine\Binaries" ReportType="console" />
diff --git a/Engine/Build/LowLevelTests/HeadlessChaos.xml b/Engine/Build/LowLevelTests/HeadlessChaos.xml
--- a/Engine/Build/LowLevelTests/HeadlessChaos.xml
+++ b/Engine/Build/LowLevelTests/HeadlessChaos.xml
@@ -1,6 +1,6 @@
 ﻿<?xml version="1.0" encoding="utf-8"?>
 <BuildGraph xmlns:xsi="http://www.w3.org/2001/XMLSchema-instance" xsi:schemaLocation="http://www.epicgames.com/BuildGraph ../../Build/Graph/Schema.xsd" xmlns="http://www.epicgames.com/BuildGraph">
-  <Option Name="RunHeadlessChaosTests" DefaultValue="False" Description="Run Headless Chaos Tests" />
+  <Option Name="RunHeadlessChaosTests" DefaultValue="" Description="Run Headless Chaos Tests" />
   <Property Name="TestNames" Value="$(TestNames);HeadlessChaos" />
   <Extend Name="RunAllTests">
     <Expand Name="DeployAndTest" Platform="Win64" TestName="HeadlessChaos" ShortName="Headless Chaos" TargetName="HeadlessChaos" BinaryRelativePath="Engine\Binaries" ReportType="console" UsesCatch2="False" />
diff --git a/Engine/Build/LowLevelTests/MathCore.xml b/Engine/Build/LowLevelTests/MathCore.xml
--- a/Engine/Build/LowLevelTests/MathCore.xml
+++ b/Engine/Build/LowLevelTests/MathCore.xml
@@ -1,6 +1,6 @@
 ﻿<?xml version="1.0" encoding="utf-8"?>
 <BuildGraph xmlns:xsi="http://www.w3.org/2001/XMLSchema-instance" xsi:schemaLocation="http://www.epicgames.com/BuildGraph ../../Build/Graph/Schema.xsd" xmlns="http://www.epicgames.com/BuildGraph">
-  <Option Name="RunMathCoreTests" DefaultValue="False" Description="Run MathCore Tests" />
+  <Option Name="RunMathCoreTests" DefaultValue="" Description="Run MathCore Tests" />
   <Property Name="TestNames" Value="$(TestNames);MathCore" />
   <Extend Name="RunAllTests">
     <Expand Name="DeployAndTest" Platform="Win64" TestName="MathCore" ShortName="MathCore" TargetName="MathCoreTests" BinaryRelativePath="Engine\Binaries" ReportType="xml" ExtraCompilationArgs="-allmodules" />
diff --git a/Engine/Build/LowLevelTests/NetworkPredictionPlugin.xml b/Engine/Build/LowLevelTests/NetworkPredictionPlugin.xml
--- a/Engine/Build/LowLevelTests/NetworkPredictionPlugin.xml
+++ b/Engine/Build/LowLevelTests/NetworkPredictionPlugin.xml
@@ -1,6 +1,6 @@
 ﻿<?xml version="1.0" encoding="utf-8"?>
 <BuildGraph xmlns:xsi="http://www.w3.org/2001/XMLSchema-instance" xsi:schemaLocation="http://www.epicgames.com/BuildGraph ../../Build/Graph/Schema.xsd" xmlns="http://www.epicgames.com/BuildGraph">
-  <Option Name="RunNetworkPredictionPluginTests" DefaultValue="False" Description="Run Net Prediction Tests" />
+  <Option Name="RunNetworkPredictionPluginTests" DefaultValue="" Description="Run Net Prediction Tests" />
   <Property Name="TestNames" Value="$(TestNames);NetworkPredictionPlugin" />
   <Extend Name="RunAllTests">
     <Expand Name="DeployAndTest" Platform="Win64" TestName="NetworkPredictionPlugin" ShortName="Net Prediction" TargetName="NetworkPredictionTests" BinaryRelativePath="Engine\Binaries" ReportType="console" />
diff --git a/Engine/Build/LowLevelTests/PlainProps.xml b/Engine/Build/LowLevelTests/PlainProps.xml
new file mode 100644
--- /dev/null
+++ b/Engine/Build/LowLevelTests/PlainProps.xml
@@ -0,0 +1,10 @@
+﻿<?xml version="1.0" encoding="utf-8"?>
+<BuildGraph xmlns:xsi="http://www.w3.org/2001/XMLSchema-instance" xsi:schemaLocation="http://www.epicgames.com/BuildGraph ../../Build/Graph/Schema.xsd" xmlns="http://www.epicgames.com/BuildGraph">
+  <Option Name="RunPlainPropsTests" DefaultValue="" Description="Run PlainProps Tests" />
+  <Property Name="TestNames" Value="$(TestNames);PlainProps" />
+  <Extend Name="RunAllTests">
+    <Expand Name="DeployAndTest" Platform="Win64" TestName="PlainProps" ShortName="PlainProps" TargetName="PlainPropsTests" BinaryRelativePath="Engine\Binaries" ReportType="xml" ExtraCompilationArgs="-allmodules" />
+    <Expand Name="DeployAndTest" Platform="Linux" TestName="PlainProps" ShortName="PlainProps" TargetName="PlainPropsTests" BinaryRelativePath="Engine\Binaries" ReportType="xml" Tags="~[.]~[Slow]" ExtraCompilationArgs="-allmodules" />
+    <Expand Name="DeployAndTest" Platform="Mac" TestName="PlainProps" ShortName="PlainProps" TargetName="PlainPropsTests" BinaryRelativePath="Engine\Binaries" ReportType="xml" ExtraCompilationArgs="-allmodules" />
+  </Extend>
+</BuildGraph>
\ No newline at end of file
diff --git a/Engine/Build/LowLevelTests/ReplicationSystem.xml b/Engine/Build/LowLevelTests/ReplicationSystem.xml
--- a/Engine/Build/LowLevelTests/ReplicationSystem.xml
+++ b/Engine/Build/LowLevelTests/ReplicationSystem.xml
@@ -1,6 +1,6 @@
 ﻿<?xml version="1.0" encoding="utf-8"?>
 <BuildGraph xmlns:xsi="http://www.w3.org/2001/XMLSchema-instance" xsi:schemaLocation="http://www.epicgames.com/BuildGraph ../../Build/Graph/Schema.xsd" xmlns="http://www.epicgames.com/BuildGraph">
-  <Option Name="RunReplicationSystemTests" DefaultValue="False" Description="Run Replication System Tests" />
+  <Option Name="RunReplicationSystemTests" DefaultValue="" Description="Run Replication System Tests" />
   <Property Name="TestNames" Value="$(TestNames);ReplicationSystem" />
   <Extend Name="RunAllTests">
     <Expand Name="DeployAndTest" Platform="Win64" TestName="ReplicationSystem" ShortName="Replication System" TargetName="ReplicationSystemLowLevelTests" BinaryRelativePath="Engine\Binaries" ReportType="console" />
diff --git a/Engine/Build/LowLevelTests/StateGraph.xml b/Engine/Build/LowLevelTests/StateGraph.xml
--- a/Engine/Build/LowLevelTests/StateGraph.xml
+++ b/Engine/Build/LowLevelTests/StateGraph.xml
@@ -1,6 +1,6 @@
 ﻿<?xml version="1.0" encoding="utf-8"?>
 <BuildGraph xmlns:xsi="http://www.w3.org/2001/XMLSchema-instance" xsi:schemaLocation="http://www.epicgames.com/BuildGraph ../../Build/Graph/Schema.xsd" xmlns="http://www.epicgames.com/BuildGraph">
-  <Option Name="RunStateGraphTests" DefaultValue="False" Description="Run State Graph Tests" />
+  <Option Name="RunStateGraphTests" DefaultValue="" Description="Run State Graph Tests" />
   <Property Name="TestNames" Value="$(TestNames);StateGraph" />
   <Extend Name="RunAllTests">
     <Expand Name="DeployAndTest" Platform="Win64" TestName="StateGraph" ShortName="State Graph" TargetName="StateGraphTests" BinaryRelativePath="Engine\Binaries" ReportType="console" />
diff --git a/Engine/Build/LowLevelTests/SwitchboardListener.xml b/Engine/Build/LowLevelTests/SwitchboardListener.xml
--- a/Engine/Build/LowLevelTests/SwitchboardListener.xml
+++ b/Engine/Build/LowLevelTests/SwitchboardListener.xml
@@ -1,6 +1,6 @@
 ﻿<?xml version="1.0" encoding="utf-8"?>
 <BuildGraph xmlns:xsi="http://www.w3.org/2001/XMLSchema-instance" xsi:schemaLocation="http://www.epicgames.com/BuildGraph ../../Build/Graph/Schema.xsd" xmlns="http://www.epicgames.com/BuildGraph">
-  <Option Name="RunSwitchboardListenerTests" DefaultValue="False" Description="Run SwitchboardListener Tests" />
+  <Option Name="RunSwitchboardListenerTests" DefaultValue="" Description="Run SwitchboardListener Tests" />
   <Property Name="TestNames" Value="$(TestNames);SwitchboardListener" />
   <Extend Name="RunAllTests">
     <Expand Name="DeployAndTest" Platform="Win64" TestName="SwitchboardListener" ShortName="SwitchboardListener" TargetName="SwitchboardListenerTests" BinaryRelativePath="Engine\Binaries" ReportType="console" />
diff --git a/Engine/Build/LowLevelTests/ToolMenus.xml b/Engine/Build/LowLevelTests/ToolMenus.xml
--- a/Engine/Build/LowLevelTests/ToolMenus.xml
+++ b/Engine/Build/LowLevelTests/ToolMenus.xml
@@ -1,6 +1,6 @@
 ﻿<?xml version="1.0" encoding="utf-8"?>
 <BuildGraph xmlns:xsi="http://www.w3.org/2001/XMLSchema-instance" xsi:schemaLocation="http://www.epicgames.com/BuildGraph ../../Build/Graph/Schema.xsd" xmlns="http://www.epicgames.com/BuildGraph">
-  <Option Name="RunToolMenusTests" DefaultValue="False" Description="Run ToolMenus Tests" />
+  <Option Name="RunToolMenusTests" DefaultValue="" Description="Run ToolMenus Tests" />
   <Property Name="TestNames" Value="$(TestNames);ToolMenus" />
   <Extend Name="RunAllTests">
     <Expand Name="DeployAndTest" Platform="Win64" TestName="ToolMenus" ShortName="ToolMenus" TargetName="ToolMenusTests" BinaryRelativePath="Engine\Binaries" ReportType="xml" ExtraCompilationArgs="-allmodules" />
diff --git a/Engine/Build/LowLevelTests/Web.xml b/Engine/Build/LowLevelTests/Web.xml
--- a/Engine/Build/LowLevelTests/Web.xml
+++ b/Engine/Build/LowLevelTests/Web.xml
@@ -1,6 +1,6 @@
 ﻿<?xml version="1.0" encoding="utf-8"?>
 <BuildGraph xmlns:xsi="http://www.w3.org/2001/XMLSchema-instance" xsi:schemaLocation="http://www.epicgames.com/BuildGraph ../../Build/Graph/Schema.xsd" xmlns="http://www.epicgames.com/BuildGraph">
-  <Option Name="RunWebTests" DefaultValue="False" Description="Run Web Tests" />
+  <Option Name="RunWebTests" DefaultValue="" Description="Run Web Tests" />
   <Property Name="TestNames" Value="$(TestNames);Web" />
   <Extend Name="RunAllTests">
     <Expand Name="DeployAndTest" Platform="Win64" TestName="Web" ShortName="Web" TargetName="WebTests" BinaryRelativePath="Engine\Binaries" ReportType="xml" HasAfterSteps="True" ExtraCompilationArgs="-allmodules" />
diff --git a/Engine/Source/Programs/UnrealBuildTool/Configuration/TestModuleRules.cs b/Engine/Source/Programs/UnrealBuildTool/Configuration/TestModuleRules.cs
--- a/Engine/Source/Programs/UnrealBuildTool/Configuration/TestModuleRules.cs
+++ b/Engine/Source/Programs/UnrealBuildTool/Configuration/TestModuleRules.cs
@@ -271,7 +271,7 @@ namespace UnrealBuildTool
 
 				if (!IsRestrictedPlatformName)
 				{
-					InsertOrUpdateTestOption(Root, $"Run{TestMetadata.TestName}Tests", $"Run {TestMetadata.TestShortName} Tests", false.ToString());
+					InsertOrUpdateTestOption(Root, $"Run{TestMetadata.TestName}Tests", $"Run {TestMetadata.TestShortName} Tests", "");
 					InsertOrUpdateTestProperty(Root, $"TestNames", TestMetadata.TestName, true);
 				}
 
