Extend the timeout and reduce the PIE iterations for CitySample TSAN tests to prevent them from timing out

#rnx

diff --git a/Engine/Build/EditorPerf.xml b/Engine/Build/EditorPerf.xml
--- a/Engine/Build/EditorPerf.xml
+++ b/Engine/Build/EditorPerf.xml
@@ -86,6 +86,7 @@
 	<!-- Option for number of PIE iterations to run (per editor launch) -->
 	<Property Name="DefaultNumIterationsPie" Value="2" />
 	<Option Name="NumIterationsPIE" DefaultValue="$(DefaultNumIterationsPIE)" Description="Specify the number of times to start/stop PIE sessions in a single editor run."/>
+	<Property Name="NumIterationsPIE" Value="1" If="$(WithTSAN)"/>
 	<Do If="'$(NumIterationsPIE)' &lt;= '0'">
 		<Warning Message="NumIterationPIE must be greater than 0. Setting to default of $(DefaultNumIterationsPIE)."/>
 		<Property Name="NumIterationsPIE" Value="$(DefaultNumIterationsPIE)" />
@@ -188,7 +189,7 @@
 	<Property Name="GauntletArgs" Value="-build=editor -tempdir=&quot;$(RootDir)/Tests&quot;"/>
 	<Property Name="GauntletArgs" Value="$(GauntletArgs) -branch=$(Branch) -ECBranch=&quot;$(Branch)&quot; -ECChangelist=&quot;$(Change)&quot;"/>
 	<Property Name="GauntletArgs" Value="$(GauntletArgs) -verbose -maxduration=7200" If="!$(WithTSAN)"/>
-	<Property Name="GauntletArgs" Value="$(GauntletArgs) -verbose -maxduration=18000" If="$(WithTSAN)"/>
+	<Property Name="GauntletArgs" Value="$(GauntletArgs) -verbose -maxduration=18000 -LogIdleTimeout=14400" If="$(WithTSAN)"/>
 
 	<!-- Gauntlet arguments for profile runs -->
 	<Property Name="ProfileArgs" Value="-tracefiletrunc"/>
