EditorPerf.xml now re-orders jobs so we always run cold then warm tests for each node type so we can have a consistent DDC when running warm nodes (and if the cold node fails we won't run the warm node at all, which was the case previously but failing cold jobs from unrelated node types could still wipe the cache for warm jobs trying to run)

#rnx

diff --git a/Engine/Build/EditorPerf.xml b/Engine/Build/EditorPerf.xml
--- a/Engine/Build/EditorPerf.xml
+++ b/Engine/Build/EditorPerf.xml
@@ -218,6 +218,7 @@
 	<!-- The main test running macro logic that is expanded with platform and dependency node names -->
 	<Property Name="TestOutputDir" Value=""/>
 	<Property Name="RunTestLogPath" Value=""/>
+	<Property Name="TraceFilePath" Value=""/>
 	<Macro Name="RunTest" Arguments="TestNameSuffix;Platform;ProjectName;ProjectTarget;NodeName;Phase;Iteration;Args;EditorArgs">
 		<Property Name="TestName" Value="EditorPerf$(Phase)$(TestNameSuffix)"/>
 		<Property Name="AssetRegistryCacheRootFolderArg" Value="-AssetRegistryCacheRootFolder=&quot;$(AssetRegistryCache_Folder)&quot;"/>
@@ -239,7 +240,8 @@
 		<Property Name="TestOutputDir" Value="$(AutomationOutputDir)/$(ProjectName)$(Platform)EditorPerf"/>
 
 		<Property Name="IterationName" Value="$(ProjectName)$(Platform)$(TestName)$(IterationSuffix)"/>
-		<Property Name="TraceFilePath" Value="$(TraceDir)/$(IterationName).utrace"/>
+		<Property Name="TraceFileName" Value="$(IterationName).utrace"/>
+		<Property Name="TraceFilePath" Value="$(TraceDir)/$(TraceFileName)"/>
 		<Property Name="TelemetryFileName" Value="$(IterationName)Telemetry.csv"/>
 		<Property Name="TelemetryFilePath" Value="$(TraceDir)/$(TelemetryFileName)"/>
 
@@ -256,16 +258,12 @@
 
 		<!-- Run the Gauntlet EditorPerf Tests -->
 		<Command Name="RunUnreal" Arguments="$(AllGauntletArgs) $(ProfileArgs) $(RunArgs) -sessionlabel=&quot;$(TestName)&quot; -tracefile=&quot;$(TraceFilePath)&quot;"/>
-
+		
 		<!-- Translate the utrace file into csv data --> 
-		<Do If="$(RunTraceSummary)">
-			<Commandlet Project="$(ProjectTarget)" Name="SummarizeTrace" Arguments="-inputfile=&quot;$(TraceFilePath)&quot; -testname=&quot;$(TestName)&quot; $(SummarizeArgs) $(CommandletArgs)"/>
-		</Do>
+		<Commandlet Project="$(ProjectTarget)" Name="SummarizeTrace" Arguments="-inputfile=&quot;$(TraceFilePath)&quot; -testname=&quot;$(TestName)&quot; $(SummarizeArgs) $(CommandletArgs)" If="$(RunTraceSummary)"/>
 		
-		<Do If="'$(TelemetryConfig)' != ''">
-			<!-- Publish telemetry to the db -->
-			<Command Name="PublishUnrealAutomationTelemetry" Arguments="-CSVFile=&quot;$(TelemetryFilePath)&quot; -project=$(Project) -changelist=$(Change) -branch=$(Branch) -TelemetryConfig=$(TelemetryConfig) $(ExtraTelemetryArgs)"/>
-		</Do>		
+		<!-- Publish telemetry to the db -->
+		<Command Name="PublishUnrealAutomationTelemetry" Arguments="-CSVFile=&quot;$(TelemetryFilePath)&quot; -project=$(Project) -changelist=$(Change) -branch=$(Branch) -TelemetryConfig=$(TelemetryConfig) $(ExtraTelemetryArgs)" If="'$(TelemetryConfig)' != ''"/>
 	</Macro>
 	
 	<Property Name="GenerateTSANReportNodeName" Value=""/>
@@ -278,8 +276,8 @@
 		</Node>
 		<Property Name="GenerateTSANReportNodeList" Value="$(GenerateTSANReportNodeList); $(GenerateTSANReportNodeName)"/>
 	</Macro>
-	
-	<Macro Name="UploadTestTraces" Arguments="Platform;ProjectName">		
+
+	<Macro Name="UploadTestTraces" Arguments="Platform;ProjectName">
 		<!-- Save all files named for this project in the trace directory -->
 		<Property Name="TestOutputDir" Value="$(AutomationOutputDir)/$(ProjectName)$(Platform)EditorPerf"/>
 		<Copy From="$(TraceDir)/$(ProjectName)$(Platform)EditorPerf*.*" To="$(TestOutputDir)"/>
@@ -358,113 +356,129 @@
 
 			<!-- If the project will be run, set up the test nodes -->
 			<Do If="ContainsItem('$(Projects)', '$(Project)', ';')=='true'">
-				<!--Each 'Warm' node in this the nested loops below depend on their 'Cold' predecessors. The nodes themselves are otherwise distinct from one another --> 
-				<Property Name="RunNodeList" Value=""/>		
+				<Property Name="RunNodeList" Value=""/>
 				<Property Name="UploadNodeList" Value=""/>
+				
 				<Agent Name="EditorPerf Agent $(Platform) $(Project)" Type="EditorPerf$(Platform)">
-					<Property Name="WarmTSANNodeRequires" Value="" />
-					
-					<ForEach Name="Phase" Values="$(Phases)">
+					<!-- TSAN Nodes -->
+					<Do If="$(WithTSAN)">
+						<!-- Force a Cold Zen Loader run to prime the DDC for our actual test case (whatever it is, zen, legacy, lazy loader etc...) -->
+						<Property Name="RunNodeColdZenName" Value="Run Non-TSAN $(OutputName) Cold ZenLoader"/>
+						<Node Name="$(RunNodeColdZenName)" Requires="$(CompileNodeName);$(SharedCompileProduces);$(CompileNodeNameProduces)" IgnoreModified="$(ProjectBinariesPath)/...">
+							<!-- Perform swap of built artifacts for TSAN builds so we use a non-TSAN build for cold runs and a TSAN build for warm runs -->
+							<Copy From="$(ProjectBinariesPath)_NonTSAN/..." To="$(ProjectBinariesPath)" ErrorIfNotFound="true" Overwrite="True"/>
+							<Expand Name="RunTest" TestNameSuffix="ZenLoader" Platform="$(Platform)" ProjectName="$(Project)" ProjectTarget="$(ProjectTarget)" NodeName="$(RunNodeColdZenName)" Phase="Cold" Iteration="1" Args="$(ColdArgs)" EditorArgs="$(ProjectCombinedArgs) $(ZenLoaderTestArgs)"/>
+						</Node>
+						<Property Name="RunNodeList" Value="$(RunNodeList);$(RunNodeColdZenName)"/>
+
+						<!-- Zen Loader TestRun -->
+						<Property Name="RunTSANNodeZenName" Value="Run TSAN $(OutputName) Warm ZenLoader"/>
+						<Node Name="$(RunTSANNodeZenName)" Requires="$(RunNodeColdZenName)" IgnoreModified="$(ProjectBinariesPath)/...">
+							<Expand Name="CopyTSANBinariesIfNeeded" ProjectBinariesPath="$(ProjectBinariesPath)"/>
+							<Expand Name="RunTest" TestNameSuffix="ZenLoader" Platform="$(Platform)" ProjectName="$(Project)" ProjectTarget="$(ProjectTarget)" NodeName="$(RunTSANNodeZenName)" Phase="Warm" Iteration="1" Args="$(WarmArgs)" EditorArgs="$(ProjectCombinedArgs) $(ZenLoaderTestArgs)"/>
+						</Node>
+						<Property Name="RunNodeList" Value="$(RunNodeList);$(RunTSANNodeZenName)"/>
+						<!-- Only depend on the Cold job since the node running TSAN likely will fail and we still want the report to run, even on failed jobs (e.g. it may have failed for TSAN reasons) -->
+						<Expand Name="GenerateTSANReport" NodeName="$(RunTSANNodeZenName)" TestOutputDir="$(TestOutputDir)" Phase="Warm" TestNameSuffix="ZenLoader" Requires="$(RunNodeColdZenName)"/>
+
+						<!-- Zen Loader Async TestRun -->
+						<Property Name="RunNodeTSANZenAsyncName" Value="Run TSAN $(OutputName) Warm ZenLoader Async"/>
+						<Node Name="$(RunNodeTSANZenAsyncName)" Requires="$(RunNodeColdZenName)" IgnoreModified="$(ProjectBinariesPath)/...">
+							<Expand Name="CopyTSANBinariesIfNeeded" ProjectBinariesPath="$(ProjectBinariesPath)"/>
+							<Expand Name="RunTest" TestNameSuffix="ZenLoaderAsync" Platform="$(Platform)" ProjectName="$(Project)" ProjectTarget="$(ProjectTarget)" NodeName="$(RunNodeTSANZenAsyncName)" Phase="Warm" Iteration="1" Args="$(WarmArgs)" EditorArgs="$(ProjectCombinedArgs) $(ZenLoaderAsyncTestArgs)"/>
+						</Node>
+						<Property Name="RunNodeList" Value="$(RunNodeList);$(RunNodeTSANZenAsyncName)"/>
+						<!-- Only depend on the Cold job since the node running TSAN likely will fail and we still want the report to run, even on failed jobs (e.g. it may have failed for TSAN reasons) -->
+						<Expand Name="GenerateTSANReport" NodeName="$(RunNodeTSANZenAsyncName)" TestOutputDir="$(TestOutputDir)" Phase="Warm" TestNameSuffix="ZenLoaderAsync" Requires="$(RunNodeColdZenName)"/>
+					</Do>
+
+					<!-- Non-TSAN Nodes -->
+					<Do If="!$(WithTSAN)">
+						<!-- 
+						Note, each of the nodes has a separate ForEach for Phases (where appropriate) because we want each node type (zen, nozen, async etc...) to be paired: Cold then Warm phase nodes.
+						If we did not order nodes like this, since the nodes run on the same agent it's possible for a cold job to fail and prevent a warm job _from an unrelated node_ to not have 
+						a properly populated DDC, defeating the purpose of the warm run. Thus all node type must run cold then warm, one after the other.
+						-->
+
 						<Property Name="RunNodeCompilationRequires" Value="$(CompileNodeName)"/>
 						<Property Name="RunNodeCompilationRequires" Value="$($($(Project)__UseBinariesFrom)__CompileNodeName)" If="'$($(Project)__UseBinariesFrom)'!=''"/>
-						<Property Name="IsWarmTSAN" Value="(('$(Phase)' == 'Warm') and $(WithTSAN))"/>
-						<Property Name="IsColdTSAN" Value="(('$(Phase)' == 'Cold') and $(WithTSAN))"/>
-						
-						<!-- Cold TSAN case -->
-						<Do If="$(IsColdTSAN)">
-							<!-- Force a Cold Zen Loader run to prime the DDC for our actual test case (whatever it is, zen, legacy, lazy loader etc...) -->
-							<Property Name="RunNodeZenName" Value="Run Non-TSAN $(OutputName) $(Phase) ZenLoader"/>
-							<Node Name="$(RunNodeZenName)" Requires="$(CompileNodeName);$(SharedCompileProduces);$(CompileNodeNameProduces)" IgnoreModified="$(ProjectBinariesPath)/...">
-								<!-- Perform swap of built artifacts for TSAN builds so we use a non-TSAN build for cold runs and a TSAN build for warm runs -->
-								<Copy From="$(ProjectBinariesPath)_NonTSAN/..." To="$(ProjectBinariesPath)" ErrorIfNotFound="true" Overwrite="True"/>
-								<Expand Name="RunTest" TestNameSuffix="ZenLoader" Platform="$(Platform)" ProjectName="$(Project)" ProjectTarget="$(ProjectTarget)" NodeName="$(RunNodeZenName)" Phase="$(Phase)" Iteration="1" Args="$($(Phase)Args)" EditorArgs="$(ProjectCombinedArgs) $(ZenLoaderTestArgs)"/>
-							</Node>
-							<Property Name="RunNodeList" Value="$(RunNodeList);$(RunNodeZenName)"/>
-							<Property Name="WarmTSANNodeRequires" Value="$(RunNodeZenName)"/>
-						</Do>
-						
-						<!-- Cold and Warm Non-TSAN case, and Warm TSAN cases -->
-						<Do If="!$(IsColdTSAN)">
-							<Property Name="RunNodeCompilationRequires" Value="$(WarmTSANNodeRequires);$(SharedCompileProduces);$(TSANCompileNodeNameProduces)" If="$(IsWarmTSAN)"/>
-							
-							<Do If="$(EnableBasePhaseTests)">
+						<Do If="$(EnableBasePhaseTests)">
+							<ForEach Name="Phase" Values="$(Phases)">
 								<ForEach Name="Iteration" Values="$(Iterations)">
 									<Property Name="IterationSuffix" Value=""/>
 									<Property Name="IterationSuffix" Value="$(Iteration)" If="'$(Iteration)' != '1'"/>
 									<Property Name="RunNodeName" Value="Run $(OutputName) $(Phase)$(IterationSuffix)"/>
 									<Property Name="RunNodeRequires" Value="$(RunNodeCompilationRequires)"/>
-									<Property Name="RunNodeRequires" Value="$(RunNodeCompilationRequires); Run $(OutputName) Cold$(IterationSuffix)" If="('$(Phase)' == 'Warm') and !$(WithTSAN)"/>
+									<Property Name="RunNodeRequires" Value="$(RunNodeCompilationRequires); Run $(OutputName) Cold$(IterationSuffix)" If="'$(Phase)' == 'Warm'"/>
 									<Node Name="$(RunNodeName)" Requires="$(RunNodeRequires)" IgnoreModified="$(ProjectBinariesPath)/...">
-										<Expand Name="CopyTSANBinariesIfNeeded" ProjectBinariesPath="$(ProjectBinariesPath)"/>
 										<Expand Name="RunTest" TestNameSuffix="NoZenLoader" Platform="$(Platform)" ProjectName="$(Project)" ProjectTarget="$(ProjectTarget)" NodeName="$(RunNodeName)" Phase="$(Phase)" Iteration="$(Iteration)" Args="$($(Phase)Args)" EditorArgs="$(ProjectCombinedArgs) $(NoZenLoaderTestArgs)"/>
 									</Node>
 									<Property Name="RunNodeList" Value="$(RunNodeList);$(RunNodeName)"/>
-									<Expand Name="GenerateTSANReport" NodeName="$(RunNodeName)" TestOutputDir="$(TestOutputDir)" Phase="$(Phase)" TestNameSuffix="NoZenLoader" Requires="$(WarmTSANNodeRequires)" If="$(IsWarmTSAN)"/>
 								</ForEach>
-							</Do>
-						
+							</ForEach>
+						</Do>
+
+						<!--Zen Loader TestRun if enabled-->
+						<Do If="$(EnableZenLoaderTest)">
 							<!-- Setup Zen Loader tests for this phase -->
-							<Do>
-								<!--Zen Loader TestRun if enabled--> 
+							<ForEach Name="Phase" Values="$(Phases)">
 								<Property Name="RunNodeZenName" Value="Run $(OutputName) $(Phase) ZenLoader"/>
 								<Property Name="RunNodeZenRequires" Value="$(RunNodeCompilationRequires)"/>
-								<Property Name="RunNodeZenRequires" Value="$(RunNodeCompilationRequires); Run $(OutputName) Cold ZenLoader" If="('$(Phase)' == 'Warm') and !$(WithTSAN)"/>
-								<Node Name="$(RunNodeZenName)" Requires="$(RunNodeZenRequires)" IgnoreModified="$(ProjectBinariesPath)/..." If="$(EnableZenLoaderTest)">
-									<Expand Name="CopyTSANBinariesIfNeeded" ProjectBinariesPath="$(ProjectBinariesPath)"/>
+								<Property Name="RunNodeZenRequires" Value="$(RunNodeCompilationRequires); Run $(OutputName) Cold ZenLoader" If="'$(Phase)' == 'Warm'"/>
+								<Node Name="$(RunNodeZenName)" Requires="$(RunNodeZenRequires)" IgnoreModified="$(ProjectBinariesPath)/...">
 									<Expand Name="RunTest" TestNameSuffix="ZenLoader" Platform="$(Platform)" ProjectName="$(Project)" ProjectTarget="$(ProjectTarget)" NodeName="$(RunNodeZenName)" Phase="$(Phase)" Iteration="1" Args="$($(Phase)Args)" EditorArgs="$(ProjectCombinedArgs) $(ZenLoaderTestArgs)"/>
 								</Node>
 								<Property Name="RunNodeList" Value="$(RunNodeList);$(RunNodeZenName)" If="$(EnableZenLoaderTest)"/>
-								<Expand Name="GenerateTSANReport" NodeName="$(RunNodeZenName)" TestOutputDir="$(TestOutputDir)" Phase="$(Phase)" TestNameSuffix="ZenLoader" Requires="$(WarmTSANNodeRequires)" If="$(EnableZenLoaderTest) And $(IsWarmTSAN)"/>
-							
-								<!--Zen Loader Async TestRun if enabled--> 
+							</ForEach>
+						</Do>
+
+						<!--Zen Loader Async TestRun if enabled-->
+						<Do If="$(EnableZenLoaderAsyncThreadTest)">
+							<ForEach Name="Phase" Values="$(Phases)">
 								<Property Name="RunNodeZenAsyncName" Value="Run $(OutputName) $(Phase) ZenLoader Async"/>
 								<Property Name="RunNodeZenAsyncRequires" Value="$(RunNodeCompilationRequires)"/>
-								<Property Name="RunNodeZenAsyncRequires" Value="$(RunNodeCompilationRequires); Run $(OutputName) Cold ZenLoader Async" If="('$(Phase)' == 'Warm') and !$(WithTSAN)"/>
-								<Node Name="$(RunNodeZenAsyncName)" Requires="$(RunNodeZenAsyncRequires)" IgnoreModified="$(ProjectBinariesPath)/..." If="$(EnableZenLoaderAsyncThreadTest)">
-									<Expand Name="CopyTSANBinariesIfNeeded" ProjectBinariesPath="$(ProjectBinariesPath)"/>
+								<Property Name="RunNodeZenAsyncRequires" Value="$(RunNodeCompilationRequires); Run $(OutputName) Cold ZenLoader Async" If="'$(Phase)' == 'Warm'"/>
+								<Node Name="$(RunNodeZenAsyncName)" Requires="$(RunNodeZenAsyncRequires)" IgnoreModified="$(ProjectBinariesPath)/...">
 									<Expand Name="RunTest" TestNameSuffix="ZenLoaderAsync" Platform="$(Platform)" ProjectName="$(Project)" ProjectTarget="$(ProjectTarget)" NodeName="$(RunNodeZenAsyncName)" Phase="$(Phase)" Iteration="1" Args="$($(Phase)Args)" EditorArgs="$(ProjectCombinedArgs) $(ZenLoaderAsyncTestArgs)"/>
 								</Node>
 								<Property Name="RunNodeList" Value="$(RunNodeList);$(RunNodeZenAsyncName)" If="$(EnableZenLoaderAsyncThreadTest)"/>
-								<Expand Name="GenerateTSANReport" NodeName="$(RunNodeZenAsyncName)" TestOutputDir="$(TestOutputDir)" Phase="$(Phase)" TestNameSuffix="ZenLoaderAsync" Requires="$(WarmTSANNodeRequires)" If="$(EnableZenLoaderAsyncThreadTest) And $(IsWarmTSAN)"/>
+							</ForEach>
+						</Do>
 
-								<!--Zen Store DDC TestRun if enabled--> 
+						<!--Zen Store DDC TestRun if enabled-->
+						<Do If="$(EnableZenStoreDDCTest)">
+							<ForEach Name="Phase" Values="$(Phases)">
 								<Property Name="RunNodeZenStoreDDCName" Value="Run $(OutputName) $(Phase) ZenLoader Store DDC"/>
 								<Property Name="RunNodeZenStoreDDCRequires" Value="$(RunNodeCompilationRequires)"/>
-								<Property Name="RunNodeZenStoreDDCRequires" Value="$(RunNodeCompilationRequires); Run $(OutputName) Cold ZenLoader Store DDC" If="('$(Phase)' == 'Warm') and !$(WithTSAN)"/>
-								<Node Name="$(RunNodeZenStoreDDCName)" Requires="$(RunNodeZenStoreDDCRequires)" IgnoreModified="$(ProjectBinariesPath)/..." If="$(EnableZenStoreDDCTest)">
-									<Expand Name="CopyTSANBinariesIfNeeded" ProjectBinariesPath="$(ProjectBinariesPath)"/>
+								<Property Name="RunNodeZenStoreDDCRequires" Value="$(RunNodeCompilationRequires); Run $(OutputName) Cold ZenLoader Store DDC" If="'$(Phase)' == 'Warm'"/>
+								<Node Name="$(RunNodeZenStoreDDCName)" Requires="$(RunNodeZenStoreDDCRequires)" IgnoreModified="$(ProjectBinariesPath)/...">
 									<Expand Name="RunTest" TestNameSuffix="ZenStoreAsyncDDC" Platform="$(Platform)" ProjectName="$(Project)" ProjectTarget="$(ProjectTarget)" NodeName="$(RunNodeZenStoreDDCName)" Phase="$(Phase)" Iteration="1" Args="$(ZenStoreDDCPhaseArgs)" EditorArgs="$(ProjectCombinedArgs) $(ZenStoreDDCTestArgs)"/>
 								</Node>
 								<Property Name="RunNodeList" Value="$(RunNodeList);$(RunNodeZenStoreDDCName)" If="$(EnableZenStoreDDCTest)"/>
-								<Expand Name="GenerateTSANReport" NodeName="$(RunNodeZenStoreDDCName)" TestOutputDir="$(TestOutputDir)" Phase="$(Phase)" TestNameSuffix="ZenStoreAsyncDDC" Requires="$(WarmTSANNodeRequires)" If="$(EnableZenStoreDDCTest) And $(IsWarmTSAN)"/>
-							</Do>
-						
-							<!--Warm only test runs e.g LazyLoad--> 
-							<Do If="'$(Phase)' == 'Warm'">
-								<!--Legacy LazyLoad TestRun if enabled--> 
-								<Property Name="RunNodeLazyLoadName" Value="Run $(OutputName) $(Phase) LazyLoad NoZenLoader"/>
-								<Property Name="RunNodeLazyLoadRequires" Value="$(RunNodeCompilationRequires)"/>
-								<Property Name="RunNodeLazyLoadRequires" Value="$(RunNodeCompilationRequires); Run $(OutputName) Cold LazyLoad NoZenLoader" If="('$(Phase)' == 'Warm') and !$(WithTSAN)"/>
-								<Node Name="$(RunNodeLazyLoadName)" Requires="$(RunNodeLazyLoadRequires)" IgnoreModified="$(ProjectBinariesPath)/..." If="$(EnableLazyLoadNoZenLoaderTest)">
-									<Expand Name="CopyTSANBinariesIfNeeded" ProjectBinariesPath="$(ProjectBinariesPath)"/>
-									<Expand Name="RunTest" TestNameSuffix="NoZenLoader" Platform="$(Platform)" ProjectName="$(Project)" ProjectTarget="$(ProjectTarget)" NodeName="$(RunNodeLazyLoadName)" Phase="$(Phase)" Iteration="1" Args="$($(Phase)Args)" EditorArgs="$(ProjectCombinedArgs) $(LazyLoadTestArgs) $(NoZenLoaderTestArgs)"/>
-								</Node>
-								<Property Name="RunNodeList" Value="$(RunNodeList);$(RunNodeLazyLoadName)" If="$(EnableLazyLoadNoZenLoaderTest)"/>
-								<Expand Name="GenerateTSANReport" NodeName="$(RunNodeLazyLoadName)" TestOutputDir="$(TestOutputDir)" Phase="$(Phase)" TestNameSuffix="NoZenLoader" Requires="$(WarmTSANNodeRequires)" If="$(EnableLazyLoadNoZenLoaderTest) And $(IsWarmTSAN)"/>
-							
-								<!--Zen LazyLoad TestRun if enabled--> 
-								<Property Name="RunNodeLazyLoadZenName" Value="Run $(OutputName) $(Phase) LazyLoad ZenLoader"/>
-								<Property Name="RunNodeLazyLoadZenRequires" Value="$(RunNodeCompilationRequires)"/>
-								<Property Name="RunNodeLazyLoadZenRequires" Value="$(RunNodeCompilationRequires); Run $(OutputName) Cold LazyLoad ZenLoader" If="('$(Phase)' == 'Warm') and !$(WithTSAN)"/>
-								<Node Name="$(RunNodeLazyLoadZenName)" Requires="$(RunNodeLazyLoadZenRequires);" IgnoreModified="$(ProjectBinariesPath)/..." If="$(EnableLazyLoadZenLoaderTest)">
-									<Expand Name="CopyTSANBinariesIfNeeded" ProjectBinariesPath="$(ProjectBinariesPath)"/>
-									<Expand Name="RunTest" TestNameSuffix="ZenLoader" Platform="$(Platform)" ProjectName="$(Project)" ProjectTarget="$(ProjectTarget)" NodeName="$(RunNodeLazyLoadZenName)" Phase="$(Phase)" Iteration="1" Args="$($(Phase)Args)" EditorArgs="$(ProjectCombinedArgs) $(LazyLoadZenTestArgs)"/>
-								</Node>
-								<Property Name="RunNodeList" Value="$(RunNodeList);$(RunNodeLazyLoadZenName)" If="$(EnableLazyLoadZenLoaderTest)"/>
-								<Expand Name="GenerateTSANReport" NodeName="$(RunNodeLazyLoadZenName)" TestOutputDir="$(TestOutputDir)" Phase="$(Phase)" TestNameSuffix="ZenLoader" Requires="$(WarmTSANNodeRequires)" If="$(EnableLazyLoadZenLoaderTest) And $(IsWarmTSAN)"/>
-							</Do>
+							</ForEach>
 						</Do>
-					</ForEach>
+
+						<!--Warm only test runs e.g LazyLoad-->
+						<ForEach Name="Phase" Values="Warm">
+							<!--Legacy LazyLoad TestRun if enabled-->
+							<Property Name="RunNodeLazyLoadName" Value="Run $(OutputName) $(Phase) LazyLoad NoZenLoader"/>
+							<Property Name="RunNodeLazyLoadRequires" Value="$(RunNodeCompilationRequires)"/>
+							<Property Name="RunNodeLazyLoadRequires" Value="$(RunNodeCompilationRequires); Run $(OutputName) Cold LazyLoad NoZenLoader" If="'$(Phase)' == 'Warm'"/>
+							<Node Name="$(RunNodeLazyLoadName)" Requires="$(RunNodeLazyLoadRequires)" IgnoreModified="$(ProjectBinariesPath)/..." If="$(EnableLazyLoadNoZenLoaderTest)">
+								<Expand Name="RunTest" TestNameSuffix="NoZenLoader" Platform="$(Platform)" ProjectName="$(Project)" ProjectTarget="$(ProjectTarget)" NodeName="$(RunNodeLazyLoadName)" Phase="$(Phase)" Iteration="1" Args="$($(Phase)Args)" EditorArgs="$(ProjectCombinedArgs) $(LazyLoadTestArgs) $(NoZenLoaderTestArgs)"/>
+							</Node>
+							<Property Name="RunNodeList" Value="$(RunNodeList);$(RunNodeLazyLoadName)" If="$(EnableLazyLoadNoZenLoaderTest)"/>
+
+							<!--Zen LazyLoad TestRun if enabled-->
+							<Property Name="RunNodeLazyLoadZenName" Value="Run $(OutputName) $(Phase) LazyLoad ZenLoader"/>
+							<Property Name="RunNodeLazyLoadZenRequires" Value="$(RunNodeCompilationRequires)"/>
+							<Property Name="RunNodeLazyLoadZenRequires" Value="$(RunNodeCompilationRequires); Run $(OutputName) Cold LazyLoad ZenLoader" If="'$(Phase)' == 'Warm'"/>
+							<Node Name="$(RunNodeLazyLoadZenName)" Requires="$(RunNodeLazyLoadZenRequires);" IgnoreModified="$(ProjectBinariesPath)/..." If="$(EnableLazyLoadZenLoaderTest)">
+								<Expand Name="RunTest" TestNameSuffix="ZenLoader" Platform="$(Platform)" ProjectName="$(Project)" ProjectTarget="$(ProjectTarget)" NodeName="$(RunNodeLazyLoadZenName)" Phase="$(Phase)" Iteration="1" Args="$($(Phase)Args)" EditorArgs="$(ProjectCombinedArgs) $(LazyLoadZenTestArgs)"/>
+							</Node>
+							<Property Name="RunNodeList" Value="$(RunNodeList);$(RunNodeLazyLoadZenName)" If="$(EnableLazyLoadZenLoaderTest)"/>
+						</ForEach>
+					</Do>
 
 					<Do If="!$(WithTSAN)">
 						<Property Name="RunNodeUploadPhaseTraces" Value="Upload $(OutputName) Traces"/>
