Fixing Build Graph duplicate product error

#rnx

diff --git a/Engine/Build/InstalledEngineBuild.xml b/Engine/Build/InstalledEngineBuild.xml
--- a/Engine/Build/InstalledEngineBuild.xml
+++ b/Engine/Build/InstalledEngineBuild.xml
@@ -768,7 +768,7 @@
 	</Agent>
 
 	<Agent Name="FeaturePacks Group" Type="Win64_Licensee">
-		<Node Name="Make Feature Packs" Requires="Update Version Files;Build Tools $(HostPlatform)">
+		<Node Name="Make Feature Packs" Requires="Update Version Files;#Build Tools $(HostPlatform)">
 			<!--Delete any pre-existing feature packs-->
 			<Delete Files="FeaturePacks/..."/>
 			<ForEach Name="TemplateName" Values="$(ProjectsToFeaturePack)">
@@ -781,7 +781,7 @@
 	</Agent>
 
 	<Agent Name="FeaturePacks Group Mac" Type="Mac_Licensee">
-		<Node Name="Make Feature Packs Mac" Requires="Update Version Files;Build Tools Mac">
+		<Node Name="Make Feature Packs Mac" Requires="Update Version Files;#Build Tools Mac">
 			<!--Delete any pre-existing feature packs-->
 			<Delete Files="FeaturePacks/..."/>
 			<ForEach Name="TemplateName" Values="$(ProjectsToFeaturePack)">
@@ -794,7 +794,7 @@
 	</Agent>
 
 	<Agent Name="FeaturePacks Group LinuxNative" Type="Linux_Licensee">
-		<Node Name="Make Feature Packs LinuxNative" Requires="Update Version Files;Build Tools LinuxNative">
+		<Node Name="Make Feature Packs LinuxNative" Requires="Update Version Files;#Build Tools LinuxNative">
 			<!--Delete any pre-existing feature packs-->
 			<Delete Files="FeaturePacks/..."/>
 			<ForEach Name="TemplateName" Values="$(ProjectsToFeaturePack)">
@@ -856,7 +856,7 @@
 	</Macro>
 
 	<Agent Name="DDC Group Win64" Type="Win64_Licensee">
-		<Node Name="Build DDC Win64" Requires="#UnrealEditor Win64;Build Tools Win64;#Build Tools CS Binaries" Produces="#CompressedDDCWin64">
+		<Node Name="Build DDC Win64" Requires="#UnrealEditor Win64;#Build Tools Win64;#Build Tools CS Binaries" Produces="#CompressedDDCWin64">
 			<!-- Get our temp dir -->
 			<Property Name="DDCDir" Value="$(RootDir)\LocalBuilds\InstalledDDC"/>
 			<Property Name="SavedDir" Value="$(SavedOutput)\Installed\Win64"/>
@@ -907,7 +907,7 @@
 
 	<Property Name="MacDDCAgentType" Value="Mac" />
 	<Agent Name="DDC Group Mac" Type="$(MacDDCAgentType)">
-		<Node Name="Build DDC Mac" Requires="#UnrealEditor Mac;Build Tools Mac;#Build Tools CS Binaries Mac" Produces="#CompressedDDCMac">
+		<Node Name="Build DDC Mac" Requires="#UnrealEditor Mac;#Build Tools Mac;#Build Tools CS Binaries Mac" Produces="#CompressedDDCMac">
 			<!-- Get our temp dir -->
 			<Property Name="DDCDir" Value="$(RootDir)/LocalBuilds/InstalledDDC"/>
 			<Property Name="SavedDir" Value="$(SavedOutput)/Installed/Mac"/>
@@ -948,7 +948,7 @@
 	</Do>
 
 	<Agent Name="DDC Group Linux" Type="Linux">
-		<Property Name="DDCLinuxRequires" Value="#UnrealEditor Linux;Build Tools LinuxNative;#Build Tools CS Binaries LinuxNative"/>
+		<Property Name="DDCLinuxRequires" Value="#UnrealEditor Linux;#Build Tools LinuxNative;#Build Tools CS Binaries LinuxNative"/>
 		<Node Name="Build DDC Linux" Requires="$(DDCLinuxRequires)" Produces="#CompressedDDCLinux">
 			<!-- Get our temp dir -->
 			<Property Name="DDCDir" Value="$(RootDir)/LocalBuilds/InstalledDDC"/>
