Tests ignores enclosing <Do if=""> condition when preflighting. Re-enforce run condition using the RunLowLevelTests command.
#rnx

diff --git a/Engine/Build/LowLevelTests.xml b/Engine/Build/LowLevelTests.xml
--- a/Engine/Build/LowLevelTests.xml
+++ b/Engine/Build/LowLevelTests.xml
@@ -222,7 +222,7 @@
 					</Do>
 
 					<Command Name="RunLowLevelTests"
-							Arguments="-test=LowLevelTests -testapp=$(TargetName)$(TestModeSuffix) -platform=$(Platform) -configuration=$(Configuration) $(TestArgsExt) -VeryVerbose" If="!$(DryRun)"/>
+							Arguments="-test=LowLevelTests -testapp=$(TargetName)$(TestModeSuffix) -platform=$(Platform) -configuration=$(Configuration) $(TestArgsExt) -VeryVerbose" If="!$(DryRun) and '$(RunUnsupported)' == ''"/>
 
 					<Do If="'$(AppRunsContainerized)' == 'True'">
 						<Property Name="TestTargetLowercase" Value=""/>
