IgnoreModified files in SanitizerReport jobs to avoid warnings from BuildGraph when running TSAN jobs

#rnx

diff --git a/Engine/Build/EditorPerf.xml b/Engine/Build/EditorPerf.xml
--- a/Engine/Build/EditorPerf.xml
+++ b/Engine/Build/EditorPerf.xml
@@ -268,10 +268,10 @@
 	
 	<Property Name="GenerateTSANReportNodeName" Value=""/>
 	<Property Name="GenerateTSANReportNodeList" Value=""/>
-	<Macro Name="GenerateTSANReport" Arguments="NodeName;TestOutputDir;Phase;TestNameSuffix;Requires">
+	<Macro Name="GenerateTSANReport" Arguments="NodeName;TestOutputDir;Phase;TestNameSuffix;Requires;IgnoreModified">
 		<Property Name="LogFilePath" Value="$(TestOutputDir)/$(Phase)$(TestNameSuffix)/Editor/EditorOutput.log"/>
 		<Property Name="GenerateTSANReportNodeName" Value="Sanitizer Report - $(NodeName)"/>
-		<Node Name="$(GenerateTSANReportNodeName)" Requires="$(Requires)">
+		<Node Name="$(GenerateTSANReportNodeName)" Requires="$(Requires)" IgnoreModified="$(IgnoreModified)">
 			<Command Name="LogSanitizerOutput" Arguments="-logfile=&quot;$(LogFilePath)&quot;" />
 		</Node>
 		<Property Name="GenerateTSANReportNodeList" Value="$(GenerateTSANReportNodeList); $(GenerateTSANReportNodeName)"/>
@@ -379,7 +379,7 @@
 						</Node>
 						<Property Name="RunNodeList" Value="$(RunNodeList);$(RunTSANNodeZenName)"/>
 						<!-- Only depend on the Cold job since the node running TSAN likely will fail and we still want the report to run, even on failed jobs (e.g. it may have failed for TSAN reasons) -->
-						<Expand Name="GenerateTSANReport" NodeName="$(RunTSANNodeZenName)" TestOutputDir="$(TestOutputDir)" Phase="Warm" TestNameSuffix="ZenLoader" Requires="$(RunNodeColdZenName)"/>
+						<Expand Name="GenerateTSANReport" NodeName="$(RunTSANNodeZenName)" TestOutputDir="$(TestOutputDir)" Phase="Warm" TestNameSuffix="ZenLoader" IgnoreModified="$(ProjectBinariesPath)/..." Requires="$(RunNodeColdZenName)"/>
 
 						<!-- Zen Loader Async TestRun -->
 						<Property Name="RunNodeTSANZenAsyncName" Value="Run TSAN $(OutputName) Warm ZenLoader Async"/>
@@ -389,7 +389,7 @@
 						</Node>
 						<Property Name="RunNodeList" Value="$(RunNodeList);$(RunNodeTSANZenAsyncName)"/>
 						<!-- Only depend on the Cold job since the node running TSAN likely will fail and we still want the report to run, even on failed jobs (e.g. it may have failed for TSAN reasons) -->
-						<Expand Name="GenerateTSANReport" NodeName="$(RunNodeTSANZenAsyncName)" TestOutputDir="$(TestOutputDir)" Phase="Warm" TestNameSuffix="ZenLoaderAsync" Requires="$(RunNodeColdZenName)"/>
+						<Expand Name="GenerateTSANReport" NodeName="$(RunNodeTSANZenAsyncName)" TestOutputDir="$(TestOutputDir)" Phase="Warm" TestNameSuffix="ZenLoaderAsync" IgnoreModified="$(ProjectBinariesPath)/..." Requires="$(RunNodeColdZenName)"/>
 					</Do>
 
 					<!-- Non-TSAN Nodes -->
