LLTs Android tests should perform full clean of all packages, currently they clean the test's package only

diff --git a/Engine/Build/LowLevelTests.xml b/Engine/Build/LowLevelTests.xml
--- a/Engine/Build/LowLevelTests.xml
+++ b/Engine/Build/LowLevelTests.xml
@@ -182,7 +182,7 @@
 					<Property Name="TestTags" Value=""/>
 					<Property Name="TestTags" Value="$(Tags)" If="'$(Tags)' != ''"/>
 					<Property Name="TestTags" Value="$(OverrideTags)" If="'$(OverrideTags)' != ''"/>
-					<Property Name="TestArgsExt" Value="$(TestArgsExt) -cleandevice -no-play-protect" If="'$(Platform)' == 'Android'"/>
+					<Property Name="TestArgsExt" Value="$(TestArgsExt) -fullclean -no-play-protect" If="'$(Platform)' == 'Android'"/>
 					<Property Name="TestArgsExt" Value="$(TestArgsExt) -tags=&quot;$(TestTags)&quot;"/>
 					<Property Name="TestArgsExt" Value="$(TestArgsExt) --buildmachine" If="$(IsBuildMachine)"/>
 					<Property Name="TestArgsExt" Value="$(TestArgsExt) -timeout=$(PerTestTimeout)"/>
