Installed Engine Build tweaks

* Improved BuildGraph Compile task batching for Win64 Tools & Datasmith compiles
* Hookup extra compile args for Linux Compile tasks
* Allow extra args for DDC steps and Feature Packs
* Add an option to compile with or with out PCH for Windows Editor/Game

#rnx

diff --git a/Engine/Build/InstalledEngineBuild.xml b/Engine/Build/InstalledEngineBuild.xml
--- a/Engine/Build/InstalledEngineBuild.xml
+++ b/Engine/Build/InstalledEngineBuild.xml
@@ -31,8 +31,13 @@
 	<Property Name="DefaultWithIOS" Value="$(DefaultWithPlatform)"/>
 	<Property Name="DefaultWithIOS" Value="false" If="'$(HostPlatform)' != 'Mac' And !$(AllPlatforms)"/>
 
-	<!-- Global Compile Arguments - note: only implemented for Win64 agents -->
-	<Option Name="ExtraCompileArgs" DefaultValue="" Description="Extra arguments to pass to compiles run on Win64 agents"/>
+	<!-- Global Compile Arguments - note: only implemented for Win64/Linux targets -->
+	<Option Name="ExtraCompileArgs" DefaultValue="" Description="Extra arguments to pass to compiles run on Win64/Linux targets"/>
+	<Option Name="ExtraDDCArgs" DefaultValue="" Description="Extra arguments to pass to DDC Generation"/>
+
+	<Option Name="WithWin64NoPCH" DefaultValue="true" Description="Compile Win64 with NoPCH"/>
+	<Property Name="Win64NoPCHArg" Value="" />
+	<Property Name="Win64NoPCHArg" Value="-NoPCH" If="$(WithWin64NoPCH)"/>
 
 	<!-- Host platforms for the installed editor -->
 	<Option Name="WithInstalledWin64" Restrict="true|false" DefaultValue="true" Description="Whether to include the Win64 editor as part of a binary release"/>
@@ -189,8 +194,8 @@
 			<SetVersion Change="$(Change)" Branch="$(EscapedBranch)" If="$(IsBuildMachine)"/>
 		</Node>
 		<Node Name="Compile UnrealEditor Win64" Requires="Update Version Files" Produces="#UnrealEditor Win64;#UnrealEditor Win64 Unstripped;#UnrealEditor Win64 Stripped;#UnrealEditor Win64 Unsigned;#UnrealEditor Win64 Signed;#UnrealEditor Win64 WithExceptions">
-			<Compile Target="UnrealEditor" Platform="Win64" Configuration="DebugGame" Tag="#UnrealEditor Win64" Arguments="-precompile -allmodules $(ExtraCompileArgs)"/>
-			<Compile Target="UnrealEditor" Platform="Win64" Configuration="Development" Tag="#UnrealEditor Win64" Arguments="-precompile -allmodules $(ExtraCompileArgs)"/>
+			<Compile Target="UnrealEditor" Platform="Win64" Configuration="DebugGame" Tag="#UnrealEditor Win64" Arguments="-precompile -allmodules $(ExtraCompileArgs) $(Win64NoPCHArg)"/>
+			<Compile Target="UnrealEditor" Platform="Win64" Configuration="Development" Tag="#UnrealEditor Win64" Arguments="-precompile -allmodules $(ExtraCompileArgs) $(Win64NoPCHArg)"/>
 
 			<Do If="$(EmbedSrcSrvInfo)">
 				<!-- Embed source info into the PDB files. Should be done from this machine to ensure that paths are correct. -->
@@ -261,8 +266,8 @@
 	<!-- Base editor and tools for compile on Linux -->
 	<Agent Name="Editor Linux" Type="Linux_Licensee">
 		<Node Name="Compile UnrealEditor Linux" Requires="Update Version Files" Produces="#UnrealEditor Linux;#UnrealEditor Linux Unstripped;#UnrealEditor Linux Stripped;#UnrealEditor Linux WithExceptions">
-			<Compile Target="UnrealEditor" Platform="Linux" Configuration="DebugGame" Tag="#UnrealEditor Linux" Arguments="-precompile -allmodules"/>
-			<Compile Target="UnrealEditor" Platform="Linux" Configuration="Development" Tag="#UnrealEditor Linux" Arguments="-precompile -allmodules"/>
+			<Compile Target="UnrealEditor" Platform="Linux" Configuration="DebugGame" Tag="#UnrealEditor Linux" Arguments="-precompile -allmodules $(ExtraCompileArgs)"/>
+			<Compile Target="UnrealEditor" Platform="Linux" Configuration="Development" Tag="#UnrealEditor Linux" Arguments="-precompile -allmodules $(ExtraCompileArgs)"/>
 
 			<Tag Files="#UnrealEditor Linux" Filter="$(LinuxStripFilter)" With="#UnrealEditor Linux Unstripped"/>
 			<Tag Files="#UnrealEditor Linux" Except="$(PluginsExceptions)" With="#UnrealEditor Linux WithExceptions"/>
@@ -287,8 +292,8 @@
 		<ForEach Name="Target" Values="$(OptionalClientTarget);$(OptionalServerTarget)">
 			<Node Name="Compile $(Target) Win64" Requires="Update Version Files" Produces="#$(Target) Win64">
 				<ForEach Name="Configuration" Values="$(GameConfigurations)">
-					<Compile Target="$(Target)" Platform="Win64" Configuration="$(Configuration)" Tag="#$(Target) Win64" Arguments="-precompile -allmodules -nolink $(TargetDebugInfoArg) $(ExtraCompileArgs)"/>
-					<Compile Target="$(Target)" Platform="Win64" Configuration="$(Configuration)" Tag="#$(Target) Win64" Arguments="-precompile $(TargetDebugInfoArg) $(ExtraCompileArgs)" Clean="false"/>
+					<Compile Target="$(Target)" Platform="Win64" Configuration="$(Configuration)" Tag="#$(Target) Win64" Arguments="-precompile -allmodules -nolink $(TargetDebugInfoArg) $(ExtraCompileArgs) $(Win64NoPCHArg)"/>
+					<Compile Target="$(Target)" Platform="Win64" Configuration="$(Configuration)" Tag="#$(Target) Win64" Arguments="-precompile $(TargetDebugInfoArg) $(ExtraCompileArgs) $(Win64NoPCHArg)" Clean="false"/>
 					<Property Name="OptionalTargetTags" Value="$(OptionalTargetTags);#$(Target) Win64"/>
 				</ForEach>
 			</Node>
@@ -298,8 +303,8 @@
 				<Tag Files="$(Tag)" With="#UnrealGame Win64"/>
 			</ForEach>
 			<ForEach Name="Configuration" Values="$(GameConfigurations)">
-				<Compile Target="UnrealGame" Platform="Win64" Configuration="$(Configuration)" Tag="#UnrealGame Win64" AllowParallelExecutor="$(bAllowParallelExecutor)" Arguments="-precompile -allmodules -nolink $(TargetDebugInfoArg) $(ExtraCompileArgs)"/>
-				<Compile Target="UnrealGame" Platform="Win64" Configuration="$(Configuration)" Tag="#UnrealGame Win64" AllowParallelExecutor="$(bAllowParallelExecutor)" Arguments="-precompile $(TargetDebugInfoArg) $(ExtraCompileArgs)" Clean="false"/>
+				<Compile Target="UnrealGame" Platform="Win64" Configuration="$(Configuration)" Tag="#UnrealGame Win64" AllowParallelExecutor="$(bAllowParallelExecutor)" Arguments="-precompile -allmodules -nolink $(TargetDebugInfoArg) $(ExtraCompileArgs) $(Win64NoPCHArg)"/>
+				<Compile Target="UnrealGame" Platform="Win64" Configuration="$(Configuration)" Tag="#UnrealGame Win64" AllowParallelExecutor="$(bAllowParallelExecutor)" Arguments="-precompile $(TargetDebugInfoArg) $(ExtraCompileArgs) $(Win64NoPCHArg)" Clean="false"/>
 			</ForEach>
 
 			<Do If="$(EmbedSrcSrvInfo)">
@@ -335,8 +340,8 @@
 		<Node Name="Compile UnrealGame Linux" Requires="Update Version Files" Produces="#UnrealGame Linux;#UnrealGame Linux Unstripped;#UnrealGame Linux Stripped">
 			<ForEach Name="Target" Values="UnrealGame;$(OptionalClientTarget);$(OptionalServerTarget)">
 				<ForEach Name="Configuration" Values="$(GameConfigurations)">
-					<Compile Target="$(Target)" Platform="Linux" Configuration="$(Configuration)" Tag="#UnrealGame Linux" AllowParallelExecutor="$(bAllowParallelExecutor)" Arguments="-precompile -allmodules -nolink $(TargetDebugInfoArg)"/>
-					<Compile Target="$(Target)" Platform="Linux" Configuration="$(Configuration)" Tag="#UnrealGame Linux" AllowParallelExecutor="$(bAllowParallelExecutor)" Arguments="-precompile $(TargetDebugInfoArg)" Clean="false"/>
+					<Compile Target="$(Target)" Platform="Linux" Configuration="$(Configuration)" Tag="#UnrealGame Linux" AllowParallelExecutor="$(bAllowParallelExecutor)" Arguments="-precompile -allmodules -nolink $(TargetDebugInfoArg) $(ExtraCompileArgs)"/>
+					<Compile Target="$(Target)" Platform="Linux" Configuration="$(Configuration)" Tag="#UnrealGame Linux" AllowParallelExecutor="$(bAllowParallelExecutor)" Arguments="-precompile $(TargetDebugInfoArg) $(ExtraCompileArgs)" Clean="false"/>
 				</ForEach>
 			</ForEach>
 
@@ -350,8 +355,8 @@
 		<Node Name="Compile UnrealGame LinuxNative" Requires="Update Version Files" Produces="#UnrealGame LinuxNative;#UnrealGame LinuxNative Unstripped;#UnrealGame LinuxNative Stripped">
 			<ForEach Name="Target" Values="UnrealGame;$(OptionalClientTarget);$(OptionalServerTarget)">
 				<ForEach Name="Configuration" Values="$(GameConfigurations)">
-					<Compile Target="$(Target)" Platform="Linux" Configuration="$(Configuration)" Tag="#UnrealGame LinuxNative" AllowParallelExecutor="$(bAllowParallelExecutor)" Arguments="-precompile -allmodules -nolink $(TargetDebugInfoArg)"/>
-					<Compile Target="$(Target)" Platform="Linux" Configuration="$(Configuration)" Tag="#UnrealGame LinuxNative" AllowParallelExecutor="$(bAllowParallelExecutor)" Arguments="-precompile $(TargetDebugInfoArg)" Clean="false"/>
+					<Compile Target="$(Target)" Platform="Linux" Configuration="$(Configuration)" Tag="#UnrealGame LinuxNative" AllowParallelExecutor="$(bAllowParallelExecutor)" Arguments="-precompile -allmodules -nolink $(TargetDebugInfoArg) $(ExtraCompileArgs)"/>
+					<Compile Target="$(Target)" Platform="Linux" Configuration="$(Configuration)" Tag="#UnrealGame LinuxNative" AllowParallelExecutor="$(bAllowParallelExecutor)" Arguments="-precompile $(TargetDebugInfoArg) $(ExtraCompileArgs)" Clean="false"/>
 				</ForEach>
 			</ForEach>
 
@@ -365,8 +370,8 @@
 		<Node Name="Compile UnrealGame LinuxArm64" Requires="Update Version Files" Produces="#UnrealGame LinuxArm64;#UnrealGame LinuxArm64 Unstripped;#UnrealGame LinuxArm64 Stripped">
 			<ForEach Name="Target" Values="UnrealGame;$(OptionalClientTarget);$(OptionalServerTarget)">
 				<ForEach Name="Configuration" Values="$(GameConfigurations)">
-					<Compile Target="$(Target)" Platform="LinuxArm64" Configuration="$(Configuration)" Tag="#UnrealGame LinuxArm64" AllowParallelExecutor="$(bAllowParallelExecutor)" Arguments="-precompile -allmodules -nolink $(TargetDebugInfoArg)"/>
-					<Compile Target="$(Target)" Platform="LinuxArm64" Configuration="$(Configuration)" Tag="#UnrealGame LinuxArm64" AllowParallelExecutor="$(bAllowParallelExecutor)" Arguments="-precompile $(TargetDebugInfoArg)" Clean="false"/>
+					<Compile Target="$(Target)" Platform="LinuxArm64" Configuration="$(Configuration)" Tag="#UnrealGame LinuxArm64" AllowParallelExecutor="$(bAllowParallelExecutor)" Arguments="-precompile -allmodules -nolink $(TargetDebugInfoArg) $(ExtraCompileArgs)"/>
+					<Compile Target="$(Target)" Platform="LinuxArm64" Configuration="$(Configuration)" Tag="#UnrealGame LinuxArm64" AllowParallelExecutor="$(bAllowParallelExecutor)" Arguments="-precompile $(TargetDebugInfoArg)  $(ExtraCompileArgs)" Clean="false"/>
 				</ForEach>
 			</ForEach>
 
@@ -380,8 +385,8 @@
 		<Node Name="Compile UnrealGame LinuxArm64 LinuxNative" Requires="Update Version Files" Produces="#UnrealGame LinuxArm64 LinuxNative;#UnrealGame LinuxArm64 LinuxNative Unstripped;#UnrealGame LinuxArm64 LinuxNative Stripped">
 			<ForEach Name="Target" Values="UnrealGame;$(OptionalClientTarget);$(OptionalServerTarget)">
 				<ForEach Name="Configuration" Values="$(GameConfigurations)">
-					<Compile Target="$(Target)" Platform="LinuxArm64" Configuration="$(Configuration)" Tag="#UnrealGame LinuxArm64 LinuxNative" AllowParallelExecutor="$(bAllowParallelExecutor)" Arguments="-precompile -allmodules -nolink $(TargetDebugInfoArg)"/>
-					<Compile Target="$(Target)" Platform="LinuxArm64" Configuration="$(Configuration)" Tag="#UnrealGame LinuxArm64 LinuxNative" AllowParallelExecutor="$(bAllowParallelExecutor)" Arguments="-precompile $(TargetDebugInfoArg)" Clean="false"/>
+					<Compile Target="$(Target)" Platform="LinuxArm64" Configuration="$(Configuration)" Tag="#UnrealGame LinuxArm64 LinuxNative" AllowParallelExecutor="$(bAllowParallelExecutor)" Arguments="-precompile -allmodules -nolink $(TargetDebugInfoArg) $(ExtraCompileArgs)"/>
+					<Compile Target="$(Target)" Platform="LinuxArm64" Configuration="$(Configuration)" Tag="#UnrealGame LinuxArm64 LinuxNative" AllowParallelExecutor="$(bAllowParallelExecutor)" Arguments="-precompile $(TargetDebugInfoArg) $(ExtraCompileArgs)" Clean="false"/>
 				</ForEach>
 			</ForEach>
 
@@ -491,7 +496,7 @@
 		</Node>
 		
 		<Node Name="Build NotForLicensee Tools Linux" Requires="Update Version Files" Produces="#NotForLicensee Build Tools Linux">
-			<Compile Target="LiveLinkHub" Configuration="Shipping" Platform="Linux" Arguments="-Monolithic $(LiveLinkHubCompileArgs)"/>
+			<Compile Target="LiveLinkHub" Configuration="Shipping" Platform="Linux" Arguments="-Monolithic $(LiveLinkHubCompileArgs) $(ExtraCompileArgs)"/>
 
 			<!-- Since we're running on the same agent as the Win64 node, we don't need to re-tag those files or else it would generate an error. -->
 			<Expand Name="ExtraNotForLicenseeFiles" TagSet="#NotForLicensee Build Tools Linux" Platform="Linux" HostPlatform="Linux"/>
@@ -501,7 +506,7 @@
 	
 	<Agent Name="NotForLicensee Tools Group LinuxNative" Type="Linux">
 		<Node Name="Build NotForLicensee Tools LinuxNative" Requires="Update Version Files" Produces="#NotForLicensee Build Tools LinuxNative">
-			<Compile Target="LiveLinkHub" Configuration="Shipping" Platform="Linux" Arguments="-Monolithic $(LiveLinkHubCompileArgs)"/>
+			<Compile Target="LiveLinkHub" Configuration="Shipping" Platform="Linux" Arguments="-Monolithic $(LiveLinkHubCompileArgs) $(ExtraCompileArgs)"/>
 
 			<Expand Name="ExtraNotForLicenseeFiles" TagSet="#NotForLicensee Build Tools LinuxNative" Platform="LinuxNative" HostPlatform="Linux"/>
 			<Tag Files="Engine/Binaries/Linux/LiveLinkHub-Linux-Shipping" With="#NotForLicensee Build Tools LinuxNative" />
@@ -532,9 +537,9 @@
 
 	<Agent Name="Tools Group Win64" Type="Win64_Licensee">
 		<Node Name="Build Tools Win64" Requires="Update Version Files;#NotForLicensee Build Tools Win64" Produces="#Tools Win64 Complete">
+			<Tag Files="#NotForLicensee Build Tools Win64" With="#Build Tools Win64"/>
 			<Compile Target="CrashReportClient" Configuration="Shipping" Platform="Win64" Tag="#Build Tools Win64" Arguments="$(CRCTelemetryDefines) $(CRCDefines) $(ExtraCompileArgs)"/>
 			<Compile Target="CrashReportClientEditor" Configuration="Shipping" Platform="Win64" Tag="#Build Tools Win64" Arguments="$(CRCETelemetryDefines) $(CRCEDefines) $(ExtraCompileArgs)"/>
-			<Tag Files="#NotForLicensee Build Tools Win64" With="#Build Tools Win64"/>
 			<Compile Target="ShaderCompileWorker" Configuration="Development" Platform="Win64" Tag="#Build Tools Win64" Arguments="$(ExtraCompileArgs)"/>
 			<Compile Target="EpicWebHelper" Configuration="Shipping" Platform="Win64" Tag="#Build Tools Win64" Arguments="$(ExtraCompileArgs)"/>
 			<Compile Target="UnrealInsights" Configuration="Development" Platform="Win64" Tag="#Build Tools Win64" Arguments="$(ExtraCompileArgs)"/>
@@ -618,22 +623,22 @@
 		</Node>
 		<Node Name="Build Tools Linux" Requires="Update Version Files;#NotForLicensee Build Tools Linux" >
 			<Tag Files="#NotForLicensee Build Tools Linux" With="#Build Tools Linux"/>
-			<Compile Target="CrashReportClient" Platform="Linux" Configuration="Shipping" Tag="#Build Tools Linux" Arguments="$(CRCTelemetryDefines)  $(CRCDefines)"/>
-			<Compile Target="CrashReportClientEditor" Platform="Linux" Configuration="Shipping" Tag="#Build Tools Linux" Arguments="$(CRCETelemetryDefines)  $(CRCEDefines)"/>
-			<Compile Target="ShaderCompileWorker" Configuration="Development" Platform="Linux" Tag="#Build Tools Linux"/>
-			<Compile Target="EpicWebHelper" Configuration="Shipping" Platform="Linux" Tag="#Build Tools Linux"/>
-			<Compile Target="UnrealInsights" Configuration="Development" Platform="Linux" Tag="#Build Tools Linux"/>
-			<Compile Target="UnrealLightmass" Configuration="Development" Platform="Linux" Tag="#Build Tools Linux"/>
-			<Compile Target="InterchangeWorker" Configuration="Development" Platform="Linux" Tag="#Build Tools Linux"/>
-			<Compile Target="UnrealPak" Configuration="Development" Platform="Linux" Tag="#Build Tools Linux"/>
-			<Compile Target="UnrealMultiUserServer" Configuration="Development" Platform="Linux" Tag="#Build Tools Linux"/>
-			<Compile Target="UnrealMultiUserSlateServer" Configuration="Development" Platform="Linux" Tag="#Build Tools Linux"/>
-			<Compile Target="BuildPatchTool" Configuration="Shipping" Platform="Linux" Tag="#Build Tools Linux"/>
-			<Compile Target="UnrealRecoverySvc" Configuration="Development" Platform="Linux" Tag="#Build Tools Linux"/>
-			<Compile Target="SwitchboardListener" Configuration="Development" Platform="Linux" Tag="#Build Tools Linux"/>
-			<Compile Target="SwitchboardListenerHelper" Configuration="Development" Platform="Linux" Tag="#Build Tools Linux"/>
-			<Compile Target="UnrealObjectPtrTool" Configuration="Development" Platform="Linux" Tag="#Build Tools Linux"/>
-			<Compile Target="ZenDashboard" Configuration="Development" Platform="Linux" Tag="#Build Tools Linux"/>
+			<Compile Target="CrashReportClient" Platform="Linux" Configuration="Shipping" Tag="#Build Tools Linux" Arguments="$(CRCTelemetryDefines)  $(CRCDefines) $(ExtraCompileArgs)"/>
+			<Compile Target="CrashReportClientEditor" Platform="Linux" Configuration="Shipping" Tag="#Build Tools Linux" Arguments="$(CRCETelemetryDefines)  $(CRCEDefines) $(ExtraCompileArgs)"/>
+			<Compile Target="ShaderCompileWorker" Configuration="Development" Platform="Linux" Tag="#Build Tools Linux" Arguments="$(ExtraCompileArgs)"/>
+			<Compile Target="EpicWebHelper" Configuration="Shipping" Platform="Linux" Tag="#Build Tools Linux" Arguments="$(ExtraCompileArgs)"/>
+			<Compile Target="UnrealInsights" Configuration="Development" Platform="Linux" Tag="#Build Tools Linux" Arguments="$(ExtraCompileArgs)"/>
+			<Compile Target="UnrealLightmass" Configuration="Development" Platform="Linux" Tag="#Build Tools Linux" Arguments="$(ExtraCompileArgs)"/>
+			<Compile Target="InterchangeWorker" Configuration="Development" Platform="Linux" Tag="#Build Tools Linux" Arguments="$(ExtraCompileArgs)"/>
+			<Compile Target="UnrealPak" Configuration="Development" Platform="Linux" Tag="#Build Tools Linux" Arguments="$(ExtraCompileArgs)"/>
+			<Compile Target="UnrealMultiUserServer" Configuration="Development" Platform="Linux" Tag="#Build Tools Linux" Arguments="$(ExtraCompileArgs)"/>
+			<Compile Target="UnrealMultiUserSlateServer" Configuration="Development" Platform="Linux" Tag="#Build Tools Linux" Arguments="$(ExtraCompileArgs)"/>
+			<Compile Target="BuildPatchTool" Configuration="Shipping" Platform="Linux" Tag="#Build Tools Linux" Arguments="$(ExtraCompileArgs)"/>
+			<Compile Target="UnrealRecoverySvc" Configuration="Development" Platform="Linux" Tag="#Build Tools Linux" Arguments="$(ExtraCompileArgs)"/>
+			<Compile Target="SwitchboardListener" Configuration="Development" Platform="Linux" Tag="#Build Tools Linux" Arguments="$(ExtraCompileArgs)"/>
+			<Compile Target="SwitchboardListenerHelper" Configuration="Development" Platform="Linux" Tag="#Build Tools Linux" Arguments="$(ExtraCompileArgs)"/>
+			<Compile Target="UnrealObjectPtrTool" Configuration="Development" Platform="Linux" Tag="#Build Tools Linux" Arguments="$(ExtraCompileArgs)"/>
+			<Compile Target="ZenDashboard" Configuration="Development" Platform="Linux" Tag="#Build Tools Linux" Arguments="$(ExtraCompileArgs)"/>
 		</Node>
 	</Agent>
 
@@ -713,22 +718,22 @@
 	<Agent Name="Tools Group LinuxNative" Type="Linux_Licensee">
 		<Node Name="Build Tools LinuxNative" Requires="Update Version Files;#NotForLicensee Build Tools LinuxNative">
 			<Tag Files="#NotForLicensee Build Tools LinuxNative" With="#Build Tools LinuxNative"/>
-			<Compile Target="CrashReportClient" Platform="Linux" Configuration="Shipping" Tag="#Build Tools LinuxNative" Arguments="$(CRCTelemetryDefines) $(CRCDefines)"/>
-			<Compile Target="CrashReportClientEditor" Platform="Linux" Configuration="Shipping"  Tag="#Build Tools LinuxNative" Arguments="$(CRCETelemetryDefines) $(CRCEDefines)"/>
-			<Compile Target="ShaderCompileWorker" Configuration="Development" Platform="Linux" Tag="#Build Tools LinuxNative"/>
-			<Compile Target="EpicWebHelper" Configuration="Shipping" Platform="Linux" Tag="#Build Tools LinuxNative"/>
-			<Compile Target="UnrealInsights" Configuration="Development" Platform="Linux" Tag="#Build Tools LinuxNative"/>
-			<Compile Target="UnrealLightmass" Configuration="Development" Platform="Linux" Tag="#Build Tools LinuxNative"/>
-			<Compile Target="InterchangeWorker" Configuration="Development" Platform="Linux" Tag="#Build Tools LinuxNative"/>
-			<Compile Target="UnrealPak" Configuration="Development" Platform="Linux" Tag="#Build Tools LinuxNative"/>
-			<Compile Target="UnrealMultiUserServer" Configuration="Development" Platform="Linux" Tag="#Build Tools LinuxNative"/>
-			<Compile Target="UnrealMultiUserSlateServer" Configuration="Development" Platform="Linux" Tag="#Build Tools LinuxNative"/>
-			<Compile Target="BuildPatchTool" Configuration="Shipping" Platform="Linux" Tag="#Build Tools LinuxNative"/>
-			<Compile Target="UnrealRecoverySvc" Configuration="Development" Platform="Linux" Tag="#Build Tools LinuxNative"/>
-			<Compile Target="SwitchboardListener" Configuration="Development" Platform="Linux" Tag="#Build Tools LinuxNative"/>
-			<Compile Target="SwitchboardListenerHelper" Configuration="Development" Platform="Linux" Tag="#Build Tools LinuxNative"/>
-			<Compile Target="UnrealObjectPtrTool" Configuration="Development" Platform="Linux" Tag="#Build Tools LinuxNative"/>
-			<Compile Target="ZenDashboard" Configuration="Development" Platform="Linux" Tag="#Build Tools LinuxNative"/>
+			<Compile Target="CrashReportClient" Platform="Linux" Configuration="Shipping" Tag="#Build Tools LinuxNative" Arguments="$(CRCTelemetryDefines) $(CRCDefines) $(ExtraCompileArgs)"/>
+			<Compile Target="CrashReportClientEditor" Platform="Linux" Configuration="Shipping"  Tag="#Build Tools LinuxNative" Arguments="$(CRCETelemetryDefines) $(CRCEDefines) $(ExtraCompileArgs)"/>
+			<Compile Target="ShaderCompileWorker" Configuration="Development" Platform="Linux" Tag="#Build Tools LinuxNative" Arguments="$(ExtraCompileArgs)"/>
+			<Compile Target="EpicWebHelper" Configuration="Shipping" Platform="Linux" Tag="#Build Tools LinuxNative" Arguments="$(ExtraCompileArgs)"/>
+			<Compile Target="UnrealInsights" Configuration="Development" Platform="Linux" Tag="#Build Tools LinuxNative" Arguments="$(ExtraCompileArgs)"/>
+			<Compile Target="UnrealLightmass" Configuration="Development" Platform="Linux" Tag="#Build Tools LinuxNative" Arguments="$(ExtraCompileArgs)"/>
+			<Compile Target="InterchangeWorker" Configuration="Development" Platform="Linux" Tag="#Build Tools LinuxNative" Arguments="$(ExtraCompileArgs)"/>
+			<Compile Target="UnrealPak" Configuration="Development" Platform="Linux" Tag="#Build Tools LinuxNative" Arguments="$(ExtraCompileArgs)"/>
+			<Compile Target="UnrealMultiUserServer" Configuration="Development" Platform="Linux" Tag="#Build Tools LinuxNative" Arguments="$(ExtraCompileArgs)"/>
+			<Compile Target="UnrealMultiUserSlateServer" Configuration="Development" Platform="Linux" Tag="#Build Tools LinuxNative" Arguments="$(ExtraCompileArgs)"/>
+			<Compile Target="BuildPatchTool" Configuration="Shipping" Platform="Linux" Tag="#Build Tools LinuxNative" Arguments="$(ExtraCompileArgs)"/>
+			<Compile Target="UnrealRecoverySvc" Configuration="Development" Platform="Linux" Tag="#Build Tools LinuxNative" Arguments="$(ExtraCompileArgs)"/>
+			<Compile Target="SwitchboardListener" Configuration="Development" Platform="Linux" Tag="#Build Tools LinuxNative" Arguments="$(ExtraCompileArgs)"/>
+			<Compile Target="SwitchboardListenerHelper" Configuration="Development" Platform="Linux" Tag="#Build Tools LinuxNative" Arguments="$(ExtraCompileArgs)"/>
+			<Compile Target="UnrealObjectPtrTool" Configuration="Development" Platform="Linux" Tag="#Build Tools LinuxNative" Arguments="$(ExtraCompileArgs)"/>
+			<Compile Target="ZenDashboard" Configuration="Development" Platform="Linux" Tag="#Build Tools LinuxNative" Arguments="$(ExtraCompileArgs)"/>
 		</Node>
 
 		<!-- Compile the C# tools. Copies files to a temporary directory first, to strip out all confidential folders and avoid sharing violations with the running version of UAT. -->
@@ -874,6 +879,7 @@
 			<Tag Files="Engine\Binaries\Win64\crashpad_handler.exe" With="#ToCopy"/>
 			<Tag Files="Engine\Binaries\Win64\zen.*" With="#ToCopy"/>
 			<Tag Files="Engine\Binaries\Win64\zenserver.*" With="#ToCopy"/>
+			<Tag Files="Engine\Binaries\Win64\UnrealBuildAccelerator\..." With="#ToCopy"/>
 			<Tag Files="Templates\TemplateResources\..." With="#ToCopy"/>
 			<Expand Name="TagDDCEngineFiles" ExceptPlugins="Plugins\...\Binaries\Mac\..."/>
 			<!-- Filter out the files not needed to build DDC. Removing confidential folders can affect DDC keys, so we want to be sure that we're making DDC with a build that can use it. -->
@@ -885,9 +891,10 @@
 
 			<!-- Force enable D3D SM6 format; this is not enabled by default in engine settings but is used by various template content -->
 			<!-- As such we must ensure we compile engine shaders for SM6 as well -->
-			<Property Name="EngineExtraArgs" Value="&quot;-ini:Engine:[/Script/WindowsTargetPlatform.WindowsTargetSettings]:D3D12TargetedShaderFormats=PCD3D_SM6&quot;"/>
+			<Property Name="EngineExtraArgs" Value="&quot;-ini:Engine:[/Script/WindowsTargetPlatform.WindowsTargetSettings]:D3D12TargetedShaderFormats=PCD3D_SM6 $(ExtraDDCArgs)&quot;"/>
+			<Property Name="FeaturePackExtraArgs" Value="&quot;$(ExtraDDCArgs)&quot;"/>
 
-			<Command Name="BuildDerivedDataCache" Arguments="-TempDir=&quot;$(DDCDir)&quot; -FeaturePacks=&quot;$(ProjectsToBuildDDCWin64)&quot; -TargetPlatforms=$(DDCPlatformsWin64) -HostPlatform=Win64 -SavedDir=&quot;$(SavedDir)&quot; -EngineContentExtraArgs=$(EngineExtraArgs)"/>
+			<Command Name="BuildDerivedDataCache" Arguments="-TempDir=&quot;$(DDCDir)&quot; -FeaturePacks=&quot;$(ProjectsToBuildDDCWin64)&quot; -TargetPlatforms=$(DDCPlatformsWin64) -HostPlatform=Win64 -SavedDir=&quot;$(SavedDir)&quot; -EngineContentExtraArgs=$(EngineExtraArgs) -FeaturePackExtraArgs=$(FeaturePackExtraArgs)"/>
 
 			<!-- Store off the generated pack csv file for later debugging -->
 			<!-- D:\build\++UE5\Sync\LocalBuilds\InstalledDDC\Engine\DerivedDataCache\ -->
@@ -986,41 +993,23 @@
 	</Agent>
 
 	<!-- These are plugins that need to be compiled with a full workspace rather than $(Platform)_Licensee. -->
-	<Agent Name="Extra Plugins Agent Win64" Type="Win64">
-		<Node Name="Compile Datasmith Plugins Win64" Requires="Update Version Files" Produces="#DatasmithPluginBinaries Win64">
-			<Do If="$(WithWin64) And $(CompileDatasmithPlugins) == 'True'">
-				<Property Name="DatasmithPlugins" Value="DatasmithCADImporter;DatasmithC4DImporter;AxFImporter;MDLImporter"/>
-				<ForEach Name="Plugin" Values="$(DatasmithPlugins)">
-					<Compile Target="UnrealEditor" Configuration="Development" Platform="Win64" Arguments="-Plugin=&quot;$(RootDir)/Engine/Plugins/Enterprise/$(Plugin)/$(Plugin).uplugin&quot; $(ExtraCompileArgs)"/>
-					<Tag Files="$(RootDir)/Engine/Plugins/Enterprise/$(Plugin)/Binaries/Win64/..." With="#DatasmithPluginBinaries Win64"/>
-				</ForEach>
-			</Do>
-		</Node>
-	</Agent>
-
-	<Agent Name="Extra Plugins Agent Mac" Type="Mac">
-		<Node Name="Compile Datasmith Plugins Mac" Requires="Update Version Files" Produces="#DatasmithPluginBinaries Mac">
-			<Do If="$(WithMac) And $(CompileDatasmithPlugins) == 'True'">
-				<Property Name="DatasmithPlugins" Value="DatasmithC4DImporter;MDLImporter"/>
+	<Macro Name="BuildDataSmithPlugins" Arguments="Platform;Plugins" OptionalArguments="CompileArgs">
+		<Agent Name="Extra Plugins Agent $(Platform)" Type="$(Platform)">
+			<Node Name="Compile Datasmith Plugins $(Platform)" Requires="Update Version Files" Produces="#DatasmithPluginBinaries $(Platform)">
+				<Property Name="DatasmithPlugins" Value="$(Plugins)"/>
+				<Property Name="TaggedPlugins" Value=""/>
 				<ForEach Name="Plugin" Values="$(DatasmithPlugins)">
-					<Compile Target="UnrealEditor" Configuration="Development" Platform="Mac" Arguments="-Plugin=&quot;$(RootDir)/Engine/Plugins/Enterprise/$(Plugin)/$(Plugin).uplugin&quot; $(MacArchitectureArg)"/>
-					<Tag Files="$(RootDir)/Engine/Plugins/Enterprise/$(Plugin)/Binaries/Mac/..." With="#DatasmithPluginBinaries Mac"/>
+					<Compile Target="UnrealEditor" Configuration="Development" Platform="$(Platform)" Arguments="-Plugin=&quot;$(RootDir)/Engine/Plugins/Enterprise/$(Plugin)/$(Plugin).uplugin&quot; $(CompileArgs)"/>
+					<Property Name="TaggedPlugins" Value="$(TaggedPlugins);$(RootDir)/Engine/Plugins/Enterprise/$(Plugin)/Binaries/$(Platform)/..."/>
 				</ForEach>
-			</Do>
-		</Node>
-	</Agent>
+				<Tag Files="$(TaggedPlugins)" With="#DatasmithPluginBinaries $(Platform)"/>
+			</Node>
+		</Agent>
+	</Macro>
 
-	<Agent Name="Extra Plugins Agent Linux" Type="Linux">
-		<Node Name="Compile Datasmith Plugins Linux" Requires="Update Version Files" Produces="#DatasmithPluginBinaries Linux">
-			<Do If="$(WithLinux) And $(CompileDatasmithPlugins) == 'True'">
-				<Property Name="DatasmithPlugins" Value="DatasmithCADImporter;MDLImporter"/>
-				<ForEach Name="Plugin" Values="$(DatasmithPlugins)">
-					<Compile Target="UnrealEditor" Configuration="Development" Platform="Linux" Arguments="-Plugin=&quot;$(RootDir)/Engine/Plugins/Enterprise/$(Plugin)/$(Plugin).uplugin&quot;"/>
-					<Tag Files="$(RootDir)/Engine/Plugins/Enterprise/$(Plugin)/Binaries/Linux/..." With="#DatasmithPluginBinaries Linux"/>
-				</ForEach>
-			</Do>
-		</Node>
-	</Agent>
+	<Expand Name="BuildDataSmithPlugins" Platform="Win64" Plugins="DatasmithCADImporter;DatasmithC4DImporter;AxFImporter;MDLImporter" If="$(WithWin64) And $(CompileDatasmithPlugins)" CompileArgs="$(ExtraCompileArgs)"/>
+	<Expand Name="BuildDataSmithPlugins" Platform="Mac" Plugins="DatasmithC4DImporter;MDLImporter" If="$(WithMac) And $(CompileDatasmithPlugins)" CompileArgs="$(MacArchitectureArg)"/>
+	<Expand Name="BuildDataSmithPlugins" Platform="Linux" Plugins="DatasmithCADImporter;MDLImporter" If="$(WithLinux) And $(CompileDatasmithPlugins)" CompileArgs="$(ExtraCompileArgs)"/>
 
 	<Agent Name="Installed Build Group Win64" Type="Win64_Licensee">
 		<!-- Build up a list of requirements as some can be optional -->
diff --git a/Engine/Source/Programs/AutomationTool/Scripts/BuildDerivedDataCache.Automation.cs b/Engine/Source/Programs/AutomationTool/Scripts/BuildDerivedDataCache.Automation.cs
--- a/Engine/Source/Programs/AutomationTool/Scripts/BuildDerivedDataCache.Automation.cs
+++ b/Engine/Source/Programs/AutomationTool/Scripts/BuildDerivedDataCache.Automation.cs
@@ -6,7 +6,6 @@ using System;
 using System.Collections.Generic;
 using System.Linq;
 using System.Text;
-using System.Threading.Tasks;
 using System.IO;
 using EpicGames.Core;
 using UnrealBuildBase;
@@ -28,6 +27,7 @@ public class BuildDerivedDataCache : BuildCommand
 		string RelativePakPath = ParseParamValue("RelativePakPath", "Engine/DerivedDataCache/Compressed.ddp");
 		bool bSkipEngine = ParseParam("SkipEngine");
 		string EngineContentExtraArgs = ParseParamValue("EngineContentExtraArgs", string.Empty);
+		string FeaturePackExtraArgs = ParseParamValue("FeaturePackExtraArgs", string.Empty);
 
 		// Get paths to everything within the temporary directory
 		string EditorExe = CommandUtils.GetEditorCommandletExe(TempDir, HostPlatform);
@@ -45,7 +45,7 @@ public class BuildDerivedDataCache : BuildCommand
 				string FeaturePackPath = CommandUtils.CombinePaths(Unreal.RootDirectory.FullName, FeaturePack);
 				if (!CommandUtils.FileExists(FeaturePackPath))
 				{
-					throw new AutomationException("Could not find project: " + FeaturePack);
+					throw new AutomationException($"Could not find project: {FeaturePack}");
 				}
 				FeaturePackPaths.Add(FeaturePackPath);
 			}
@@ -54,11 +54,18 @@ public class BuildDerivedDataCache : BuildCommand
 		// loop through all the paths and generate ddc data for them
 		foreach (string FeaturePackPath in FeaturePackPaths)
 		{
+			StringBuilder FeaturePackParameters = new StringBuilder();
 			string ProjectSpecificPlatforms = TargetPlatforms;
 			FileReference FileRef = new FileReference(FeaturePackPath);
 			string GameName = FileRef.GetFileNameWithoutAnyExtensions();
 			ProjectDescriptor Project = ProjectDescriptor.FromFile(FileRef);
 
+			FeaturePackParameters.Append($"-fill -DDC={BackendName} -ProjectOnly");
+			if (!string.IsNullOrEmpty(FeaturePackExtraArgs))
+			{
+				FeaturePackParameters.Append($" {FeaturePackExtraArgs} ");
+			}
+
 			if (Project.TargetPlatforms != null && Project.TargetPlatforms.Length > 0)
 			{
 				// Restrict target platforms used to those specified in project file
@@ -86,9 +93,9 @@ public class BuildDerivedDataCache : BuildCommand
 				ProjectSpecificPlatforms = CommandUtils.CombineCommandletParams(FilteredPlatforms.Distinct().ToArray());
 			}
 			Logger.LogInformation("Generating DDC data for {GameName} on {ProjectSpecificPlatforms}", GameName, ProjectSpecificPlatforms);
-			CommandUtils.DDCCommandlet(FileRef, EditorExe, null, ProjectSpecificPlatforms, String.Format("-fill -DDC={0} -ProjectOnly", BackendName));
+			CommandUtils.DDCCommandlet(FileRef, EditorExe, null, ProjectSpecificPlatforms, FeaturePackParameters.ToString());
 
-			string ProjectPakFile = CommandUtils.CombinePaths(Path.GetDirectoryName(OutputPakFile), String.Format("Compressed-{0}.ddp", GameName));
+			string ProjectPakFile = CommandUtils.CombinePaths(Path.GetDirectoryName(OutputPakFile), $"Compressed-{GameName}.ddp");
 			CommandUtils.DeleteFile(ProjectPakFile);
 			CommandUtils.RenameFile(OutputPakFile, ProjectPakFile);
 
