Include UBA for other platforms

#rnx

diff --git a/Engine/Build/InstalledEngineFilters.xml b/Engine/Build/InstalledEngineFilters.xml
--- a/Engine/Build/InstalledEngineFilters.xml
+++ b/Engine/Build/InstalledEngineFilters.xml
@@ -390,6 +390,9 @@
 		Engine/Binaries/Mac/zenserver
 		Engine/Binaries/Mac/zen
 		Engine/Binaries/Mac/crashpad_handler
+
+		<!-- Precompiled Mac UBA -->
+		Engine/Binaries/Mac/UnrealBuildAccelerator/...
 			
 		<!-- Catch2 -->
 		Engine/Source/ThirdParty/Catch2/v3.4.0/...
@@ -437,6 +440,8 @@
 		Templates/TP_nDisplayBP/...
 		Templates/TP_InCamVFXBP/...
 
+		<!-- Precompiled Linux UBA -->
+		Engine/Binaries/Linux/UnrealBuildAccelerator/...
 
 		<!-- DotNet -->
 		Engine/Binaries/ThirdParty/DotNet/8.0.300/linux-x64/...
