Whitelisting documentation files for Network Prediction and Mover plugins

diff --git a/Engine/Build/InstalledEngineFilters.xml b/Engine/Build/InstalledEngineFilters.xml
--- a/Engine/Build/InstalledEngineFilters.xml
+++ b/Engine/Build/InstalledEngineFilters.xml
@@ -149,6 +149,11 @@
 		
 		<!-- Source code and demos for container images -->
 		Engine/Extras/Containers/...
+
+		<!-- Additional documentation -->
+		Engine/Plugins/Runtime/NetworkPrediction/readme.txt
+		Engine/Plugins/Experimental/Mover/README.md
+
 	</Property>
 	
 	<!-- Files which can exist under any engine or engine-platform directory -->
