Adding SteamDeck python file to Win64/Linux Installed builds

#rnx

diff --git a/Engine/Build/InstalledEngineFilters.xml b/Engine/Build/InstalledEngineFilters.xml
--- a/Engine/Build/InstalledEngineFilters.xml
+++ b/Engine/Build/InstalledEngineFilters.xml
@@ -289,6 +289,9 @@
 		<!-- Default resource files -->
 		Engine/Build/Windows/Resources/...
 
+		<!-- Steam Deck files -->
+		Engine/Build/SteamDeck/...
+
 		<!-- Binaries which are explicitly staged by UAT rather than listing runtime dependencies -->
 		Engine/Binaries/ThirdParty/OpenSSL/...
 
@@ -438,6 +441,9 @@
 		Engine/Build/BatchFiles/RunUAT.sh
 		Engine/Build/BatchFiles/RunUBT.sh
 
+		<!-- Steam Deck files -->
+		Engine/Build/SteamDeck/...
+
 		<!-- Debug tools -->
 		Engine/Binaries/Linux/dump_syms
 		Engine/Binaries/Linux/BreakpadSymbolEncoder
