Add -uba to the compile portions of Low Level Tests when the option is enabled
Remove NoXGE variable
#rnx

diff --git a/Engine/Build/LowLevelTests.xml b/Engine/Build/LowLevelTests.xml
--- a/Engine/Build/LowLevelTests.xml
+++ b/Engine/Build/LowLevelTests.xml
@@ -76,6 +76,9 @@
 	<!-- Test Options -->
 	<Option Name="RunAllTests" DefaultValue="true" Description="Run All Available Tests (except Foundation)"/>
 	<Option Name="DryRun" DefaultValue="false" Description="Dry Run"/>
+
+	<!-- Additional Compile Options -->
+	<Option Name="WithUBA" DefaultValue="true" Description="Use UBA for compile task"/>
 	
 	<!-- Platform Options -->
 	<Option Name="TestAllSupportedPlatforms" DefaultValue="true" Description="Test All Supported Platforms"/>
@@ -138,6 +141,7 @@
 		<Property Name="LocalAchiveDirectory" Value="$(RootDir)/$(TargetName)"/>
 		<Property Name="ExtraCompilationArgsFinal" Value="$(ExtraCompilationArgs)"/>
 		<Do If="'$(StagesWithProjectFile)' == 'True'">
+		<Property Name="ExtraCompilationArgsFinal" Value="$(ExtraCompilationArgsFinal) -ubtargs=&quot;-UBA &quot;" If="$(WithUBA)"/>
 			<Do If="'$(BuildProject)' != ''">
 				<Property Name="ExtraCompilationArgsFinal" Value="$(ExtraCompilationArgs) -projectprogramoverride=$(RootDir)/$(BuildProject)"/>
 			</Do>
@@ -147,7 +151,7 @@
 			<Do If="'$(BuildProject)' != ''">
 				<Property Name="ExtraCompilationArgsFinal" Value="$(ExtraCompilationArgs) -Project=$(RootDir)/$(BuildProject)"/>
 			</Do>
-			<Compile Target="$(TargetName)" Configuration="$(Configuration)" Platform="$(Platform)" AllowXGE="false" Arguments="-Mode=Test $(ExtraCompilationArgsFinal)" If="!$(DryRun)"/>
+			<Compile Target="$(TargetName)" Configuration="$(Configuration)" Platform="$(Platform)" Arguments="-Mode=Test $(ExtraCompilationArgsFinal)" If="!$(DryRun)"/>
 		</Do>
 	</Macro>
 	
