Don't publish TSAN telemetry to the EditorPerf DBs

#rnx

diff --git a/Engine/Build/EditorPerf.xml b/Engine/Build/EditorPerf.xml
--- a/Engine/Build/EditorPerf.xml
+++ b/Engine/Build/EditorPerf.xml
@@ -222,16 +222,21 @@
 	<Macro Name="RunTest" Arguments="TestNameSuffix;Platform;ProjectName;ProjectTarget;NodeName;Phase;Iteration;Args;EditorArgs">
 		<Property Name="TestName" Value="EditorPerf$(Phase)$(TestNameSuffix)"/>
 		<Property Name="AssetRegistryCacheRootFolderArg" Value="-AssetRegistryCacheRootFolder=&quot;$(AssetRegistryCache_Folder)&quot;"/>
-
+		
+		<!-- Don't send telemetry for TSAN jobs as the time won't be comparable -->
+		<Property Name="EditorTelemetryArgs" Value=""/>
+		<Property Name="EditorTelemetryArgs" Value="-ST_SendTelemetry=false" If="$(WithTSAN)"/>
+		
 		<Property Name="TestFilterArgs" Value=""/>
 		<Property Name="CombinedEditorArgs" Value=""/>
+	
 		<Do If="'$($(ProjectName)__GauntletTest)'=='EditorTests'">
 			<Property Name="TestFilter" Value="Project.Iteration.PIE"/>
 			<Property Name="TestFilterArgs" Value="-runtest=$(TestFilter)"/>
-			<Property Name="CombinedEditorArgs" Value="$(EditorArgs) $(TestFilterArgs) $(AssetRegistryCacheRootFolderArg) $(ArgIterationsPIE)"/>
+			<Property Name="CombinedEditorArgs" Value="$(EditorArgs) $(TestFilterArgs) $(AssetRegistryCacheRootFolderArg) $(ArgIterationsPIE) $(EditorTelemetryArgs)"/>
 		</Do>
 		<Do If="'$($(ProjectName)__GauntletTest)'!='EditorTests'">
-			<Property Name="CombinedEditorArgs" Value="$(EditorArgs) $(AssetRegistryCacheRootFolderArg)"/>
+			<Property Name="CombinedEditorArgs" Value="$(EditorArgs) $(AssetRegistryCacheRootFolderArg) $(EditorTelemetryArgs)"/>
 		</Do>
 
 		<Property Name="IterationSuffix" Value=""/>
