Allow supplying extra commandline args for Mac compiles

#rnx

diff --git a/Engine/Build/InstalledEngineBuild.xml b/Engine/Build/InstalledEngineBuild.xml
--- a/Engine/Build/InstalledEngineBuild.xml
+++ b/Engine/Build/InstalledEngineBuild.xml
@@ -32,7 +32,8 @@
 	<Property Name="DefaultWithIOS" Value="false" If="'$(HostPlatform)' != 'Mac' And !$(AllPlatforms)"/>
 
 	<!-- Global Compile Arguments - note: only implemented for Win64/Linux targets -->
-	<Option Name="ExtraCompileArgs" DefaultValue="" Description="Extra arguments to pass to compiles run on Win64/Linux targets"/>
+	<Option Name="ExtraCompileArgs" DefaultValue="" Description="Extra arguments to pass to compiles run on non Mac based platforms"/>
+	<Option Name="ExtraCompileArgsMac" DefaultValue="" Description="Extra arguments to pass to compiles run on Mac platforms"/>
 	<Option Name="ExtraDDCArgs" DefaultValue="" Description="Extra arguments to pass to DDC Generation"/>
 
 	<Option Name="WithWin64NoPCH" DefaultValue="true" Description="Compile Win64 with NoPCH"/>
@@ -234,8 +235,8 @@
 	<Agent Name="Editor Mac" Type="Mac_Licensee">
 		<Node Name="Compile UnrealEditor Mac" Requires="Update Version Files" Produces="#UnrealEditor Mac;#UnrealEditor Mac Unstripped;#UnrealEditor Mac Stripped;#UnrealEditor Mac Unsigned;#UnrealEditor Mac Signed;#UnrealEditor Mac WithExceptions">
 			<Property Name="AllModulesArg" Value="-allmodules"/>
-			<Compile Target="UnrealEditor" Platform="Mac" Configuration="DebugGame" Tag="#UnrealEditor Mac" Arguments="-precompile -EnableDSYM $(AllModulesArg) $(MacArchitectureArg)"/>
-			<Compile Target="UnrealEditor" Platform="Mac" Configuration="Development" Tag="#UnrealEditor Mac" Arguments="-precompile -EnableDSYM $(AllModulesArg) $(MacArchitectureArg)"/>
+			<Compile Target="UnrealEditor" Platform="Mac" Configuration="DebugGame" Tag="#UnrealEditor Mac" Arguments="-precompile -EnableDSYM $(ExtraCompileArgsMac) $(AllModulesArg) $(MacArchitectureArg)"/>
+			<Compile Target="UnrealEditor" Platform="Mac" Configuration="Development" Tag="#UnrealEditor Mac" Arguments="-precompile -EnableDSYM $(ExtraCompileArgsMac) $(AllModulesArg) $(MacArchitectureArg)"/>
 
 			<Property Name="UnrealEditorExceptions" Value=""/>
 			<Property Name="UnrealEditorExceptions" Value="$(PluginsExceptions)" If="$(CompileDatasmithPlugins)"/>
@@ -419,8 +420,8 @@
 		<ForEach Name="Target" Values="$(OptionalClientTarget);$(OptionalServerTarget)">
 			<Node Name="Compile $(Target) Mac" Requires="Update Version Files" Produces="#$(Target) Mac">
 				<ForEach Name="Configuration" Values="$(GameConfigurations)">
-					<Compile Target="$(Target)" Platform="Mac" Configuration="$(Configuration)" Tag="#$(Target) Mac" Arguments="-precompile -allmodules -nolink $(TargetDebugInfoArgMac)"/>
-					<Compile Target="$(Target)" Platform="Mac" Configuration="$(Configuration)" Tag="#$(Target) Mac" Arguments="-precompile $(TargetDebugInfoArgMac)" Clean="false"/>
+					<Compile Target="$(Target)" Platform="Mac" Configuration="$(Configuration)" Tag="#$(Target) Mac" Arguments="-precompile -allmodules -nolink $(ExtraCompileArgsMac) $(TargetDebugInfoArgMac)"/>
+					<Compile Target="$(Target)" Platform="Mac" Configuration="$(Configuration)" Tag="#$(Target) Mac" Arguments="-precompile $(ExtraCompileArgsMac) $(TargetDebugInfoArgMac)" Clean="false"/>
 					<Property Name="OptionalTargetTags" Value="$(OptionalTargetTags);#$(Target) Mac"/>
 				</ForEach>
 			</Node>
@@ -431,8 +432,8 @@
 				<Tag Files="$(Tag)" With="#UnrealGame Mac"/>
 			</ForEach>
 			<ForEach Name="Configuration" Values="$(GameConfigurations)">
-				<Compile Target="UnrealGame" Platform="Mac" Configuration="$(Configuration)" Tag="#UnrealGame Mac" AllowParallelExecutor="$(bAllowParallelExecutor)" Arguments="-precompile -allmodules -nolink $(TargetDebugInfoArgMac)"/>
-				<Compile Target="UnrealGame" Platform="Mac" Configuration="$(Configuration)" Tag="#UnrealGame Mac" AllowParallelExecutor="$(bAllowParallelExecutor)" Arguments="-precompile $(TargetDebugInfoArgMac)" Clean="false"/>
+				<Compile Target="UnrealGame" Platform="Mac" Configuration="$(Configuration)" Tag="#UnrealGame Mac" AllowParallelExecutor="$(bAllowParallelExecutor)" Arguments="-precompile -allmodules -nolink $(ExtraCompileArgsMac) $(TargetDebugInfoArgMac)"/>
+				<Compile Target="UnrealGame" Platform="Mac" Configuration="$(Configuration)" Tag="#UnrealGame Mac" AllowParallelExecutor="$(bAllowParallelExecutor)" Arguments="-precompile $(ExtraCompileArgsMac) $(TargetDebugInfoArgMac)" Clean="false"/>
 			</ForEach>
 
 			<Tag Files="#UnrealGame Mac" Filter="$(MacStripFilter)" With="#UnrealGame Mac Unstripped"/>
@@ -451,8 +452,8 @@
 		<Node Name="Compile UnrealGame IOS" Requires="Update Version Files" Produces="#UnrealGame IOS;#UnrealGame IOS Unstripped;#UnrealGame IOS Stripped">
 			<ForEach Name="Target" Values="UnrealGame;$(OptionalClientTarget)">
 				<ForEach Name="Configuration" Values="$(GameConfigurations)">
-					<Compile Target="$(Target)" Platform="IOS" Configuration="$(Configuration)" Tag="#UnrealGame IOS" AllowParallelExecutor="$(bAllowParallelExecutor)" Arguments="-precompile -allmodules -nolink $(TargetDebugInfoArgMac)"/>
-					<Compile Target="$(Target)" Platform="IOS" Configuration="$(Configuration)" Tag="#UnrealGame IOS" AllowParallelExecutor="$(bAllowParallelExecutor)" Arguments="-precompile -createstub $(TargetDebugInfoArgMac)" Clean="false"/>
+					<Compile Target="$(Target)" Platform="IOS" Configuration="$(Configuration)" Tag="#UnrealGame IOS" AllowParallelExecutor="$(bAllowParallelExecutor)" Arguments="-precompile -allmodules -nolink $(ExtraCompileArgsMac) $(TargetDebugInfoArgMac)"/>
+					<Compile Target="$(Target)" Platform="IOS" Configuration="$(Configuration)" Tag="#UnrealGame IOS" AllowParallelExecutor="$(bAllowParallelExecutor)" Arguments="-precompile -createstub $(ExtraCompileArgsMac) $(TargetDebugInfoArgMac)" Clean="false"/>
 				</ForEach>
 			</ForEach>
 
@@ -466,8 +467,8 @@
 		<Node Name="Compile UnrealGame TVOS" Requires="Update Version Files" Produces="#UnrealGame TVOS;#UnrealGame TVOS Unstripped;#UnrealGame TVOS Stripped">
 			<ForEach Name="Target" Values="UnrealGame;$(OptionalClientTarget)">
 				<ForEach Name="Configuration" Values="$(GameConfigurations)">
-					<Compile Target="$(Target)" Platform="TVOS" Configuration="$(Configuration)" Tag="#UnrealGame TVOS" AllowParallelExecutor="$(bAllowParallelExecutor)" Arguments="-precompile -allmodules -nolink $(TargetDebugInfoArgMac)"/>
-					<Compile Target="$(Target)" Platform="TVOS" Configuration="$(Configuration)" Tag="#UnrealGame TVOS" AllowParallelExecutor="$(bAllowParallelExecutor)" Arguments="-precompile -createstub $(TargetDebugInfoArgMac)" Clean="false"/>
+					<Compile Target="$(Target)" Platform="TVOS" Configuration="$(Configuration)" Tag="#UnrealGame TVOS" AllowParallelExecutor="$(bAllowParallelExecutor)" Arguments="-precompile -allmodules -nolink $(ExtraCompileArgsMac) $(TargetDebugInfoArgMac)"/>
+					<Compile Target="$(Target)" Platform="TVOS" Configuration="$(Configuration)" Tag="#UnrealGame TVOS" AllowParallelExecutor="$(bAllowParallelExecutor)" Arguments="-precompile -createstub $(ExtraCompileArgsMac) $(TargetDebugInfoArgMac)" Clean="false"/>
 				</ForEach>
 			</ForEach>
 
@@ -516,7 +517,7 @@
 	
 	<Agent Name="NotForLicensee Tools Group Mac" Type="Mac">
 		<Node Name="Build NotForLicensee Tools Mac" Requires="Update Version Files" Produces="#NotForLicensee Build Tools Mac">
-			<Compile Target="LiveLinkHub" Configuration="Shipping" Platform="Mac" Arguments="-Monolithic $(LiveLinkHubCompileArgs)"/>
+			<Compile Target="LiveLinkHub" Configuration="Shipping" Platform="Mac" Arguments="-Monolithic $(LiveLinkHubCompileArgs) $(ExtraCompileArgsMac)"/>
 
 			<Expand Name="ExtraNotForLicenseeFiles" TagSet="#NotForLicensee Build Tools Mac" Platform="Mac" HostPlatform="Mac"/>
 			<Tag Files="Engine/Binaries/Mac/LiveLinkHub-Mac-Shipping.app/..." With="#NotForLicensee Build Tools Mac" />
@@ -694,24 +695,24 @@
 	
 		<Node Name="Build Tools Mac" Requires="Update Version Files;#NotForLicensee Build Tools Mac">
 			<Tag Files="#NotForLicensee Build Tools Mac" With="#Build Tools Mac"/>
-			<Compile Target="CrashReportClient" Configuration="Shipping" Platform="Mac" Tag="#Build Tools Mac" Arguments="-EnableDSYM $(CRCTelemetryDefines) $(CRCDefines)"/>
-			<Compile Target="CrashReportClientEditor" Configuration="Shipping" Platform="Mac" Tag="#Build Tools Mac" Arguments="-EnableDSYM $(CRCETelemetryDefines) $(CRCEDefines)"/>
-			<Compile Target="ShaderCompileWorker" Configuration="Development" Platform="Mac" Tag="#Build Tools Mac" Arguments="-EnableDSYM"/>
-			<Compile Target="EpicWebHelper" Configuration="Shipping" Platform="Mac" Tag="#Build Tools Mac" Arguments="-EnableDSYM"/>
-			<Compile Target="UnrealInsights" Configuration="Development" Platform="Mac" Tag="#Build Tools Mac" Arguments="-EnableDSYM"/>
-			<Compile Target="UnrealFrontend" Configuration="Development" Platform="Mac" Tag="#Build Tools Mac" Arguments="-EnableDSYM"/>
-			<Compile Target="UnrealLightmass" Configuration="Development" Platform="Mac" Tag="#Build Tools Mac" Arguments="-EnableDSYM"/>
-			<Compile Target="InterchangeWorker" Configuration="Development" Platform="Mac" Tag="#Build Tools Mac" Arguments="-EnableDSYM"/>
-			<Compile Target="UnrealPak" Configuration="Development" Platform="Mac" Tag="#Build Tools Mac" Arguments="-EnableDSYM"/>
-			<Compile Target="UnrealMultiUserServer" Configuration="Development" Platform="Mac" Tag="#Build Tools Mac" Arguments="-EnableDSYM"/>
-			<Compile Target="UnrealMultiUserSlateServer" Configuration="Development" Platform="Mac" Tag="#Build Tools Mac" Arguments="-EnableDSYM"/>
-			<Compile Target="UnrealRecoverySvc" Configuration="Development" Platform="Mac" Tag="#Build Tools Mac" Arguments="-EnableDSYM"/>
-			<Compile Target="BuildPatchTool" Configuration="Shipping" Platform="Mac" Tag="#Build Tools Mac" Arguments="-EnableDSYM"/>
-			<Compile Target="DsymExporter" Configuration="Development" Platform="Mac" Tag="#Build Tools Mac" Arguments="-EnableDSYM"/>
-			<Compile Target="UnrealEditorServices" Configuration="Development" Platform="Mac" Tag="#Build Tools Mac" Arguments="-EnableDSYM"/>
-			<Compile Target="UnrealAtoS" Configuration="Development" Platform="Mac" Tag="#Build Tools Mac" Arguments="-EnableDSYM"/>
-			<Compile Target="UnrealObjectPtrTool" Configuration="Development" Platform="Mac" Tag="#Build Tools Mac" Arguments="-EnableDSYM"/>
-			<Compile Target="ZenDashboard" Configuration="Development" Platform="Mac" Tag="#Build Tools Mac" Arguments="-EnableDSYM"/>
+			<Compile Target="CrashReportClient" Configuration="Shipping" Platform="Mac" Tag="#Build Tools Mac" Arguments="-EnableDSYM $(CRCTelemetryDefines) $(CRCDefines) $(ExtraCompileArgsMac)"/>
+			<Compile Target="CrashReportClientEditor" Configuration="Shipping" Platform="Mac" Tag="#Build Tools Mac" Arguments="-EnableDSYM $(CRCETelemetryDefines) $(CRCEDefines) $(ExtraCompileArgsMac)"/>
+			<Compile Target="ShaderCompileWorker" Configuration="Development" Platform="Mac" Tag="#Build Tools Mac" Arguments="-EnableDSYM $(ExtraCompileArgsMac)"/>
+			<Compile Target="EpicWebHelper" Configuration="Shipping" Platform="Mac" Tag="#Build Tools Mac" Arguments="-EnableDSYM $(ExtraCompileArgsMac)"/>
+			<Compile Target="UnrealInsights" Configuration="Development" Platform="Mac" Tag="#Build Tools Mac" Arguments="-EnableDSYM $(ExtraCompileArgsMac)"/>
+			<Compile Target="UnrealFrontend" Configuration="Development" Platform="Mac" Tag="#Build Tools Mac" Arguments="-EnableDSYM $(ExtraCompileArgsMac)"/>
+			<Compile Target="UnrealLightmass" Configuration="Development" Platform="Mac" Tag="#Build Tools Mac" Arguments="-EnableDSYM $(ExtraCompileArgsMac)"/>
+			<Compile Target="InterchangeWorker" Configuration="Development" Platform="Mac" Tag="#Build Tools Mac" Arguments="-EnableDSYM $(ExtraCompileArgsMac)"/>
+			<Compile Target="UnrealPak" Configuration="Development" Platform="Mac" Tag="#Build Tools Mac" Arguments="-EnableDSYM $(ExtraCompileArgsMac)"/>
+			<Compile Target="UnrealMultiUserServer" Configuration="Development" Platform="Mac" Tag="#Build Tools Mac" Arguments="-EnableDSYM $(ExtraCompileArgsMac)"/>
+			<Compile Target="UnrealMultiUserSlateServer" Configuration="Development" Platform="Mac" Tag="#Build Tools Mac" Arguments="-EnableDSYM $(ExtraCompileArgsMac)"/>
+			<Compile Target="UnrealRecoverySvc" Configuration="Development" Platform="Mac" Tag="#Build Tools Mac" Arguments="-EnableDSYM $(ExtraCompileArgsMac)"/>
+			<Compile Target="BuildPatchTool" Configuration="Shipping" Platform="Mac" Tag="#Build Tools Mac" Arguments="-EnableDSYM $(ExtraCompileArgsMac)"/>
+			<Compile Target="DsymExporter" Configuration="Development" Platform="Mac" Tag="#Build Tools Mac" Arguments="-EnableDSYM $(ExtraCompileArgsMac)"/>
+			<Compile Target="UnrealEditorServices" Configuration="Development" Platform="Mac" Tag="#Build Tools Mac" Arguments="-EnableDSYM $(ExtraCompileArgsMac)"/>
+			<Compile Target="UnrealAtoS" Configuration="Development" Platform="Mac" Tag="#Build Tools Mac" Arguments="-EnableDSYM $(ExtraCompileArgsMac)"/>
+			<Compile Target="UnrealObjectPtrTool" Configuration="Development" Platform="Mac" Tag="#Build Tools Mac" Arguments="-EnableDSYM $(ExtraCompileArgsMac)"/>
+			<Compile Target="ZenDashboard" Configuration="Development" Platform="Mac" Tag="#Build Tools Mac" Arguments="-EnableDSYM $(ExtraCompileArgsMac)"/>
 		</Node>
 	</Agent>
 
@@ -1008,7 +1009,7 @@
 	</Macro>
 
 	<Expand Name="BuildDataSmithPlugins" Platform="Win64" Plugins="DatasmithCADImporter;DatasmithC4DImporter;AxFImporter;MDLImporter" If="$(WithWin64) And $(CompileDatasmithPlugins)" CompileArgs="$(ExtraCompileArgs)"/>
-	<Expand Name="BuildDataSmithPlugins" Platform="Mac" Plugins="DatasmithC4DImporter;MDLImporter" If="$(WithMac) And $(CompileDatasmithPlugins)" CompileArgs="$(MacArchitectureArg)"/>
+	<Expand Name="BuildDataSmithPlugins" Platform="Mac" Plugins="DatasmithC4DImporter;MDLImporter" If="$(WithMac) And $(CompileDatasmithPlugins)" CompileArgs="$(MacArchitectureArg) $(ExtraCompileArgsMac)"/>
 	<Expand Name="BuildDataSmithPlugins" Platform="Linux" Plugins="DatasmithCADImporter;MDLImporter" If="$(WithLinux) And $(CompileDatasmithPlugins)" CompileArgs="$(ExtraCompileArgs)"/>
 
 	<Agent Name="Installed Build Group Win64" Type="Win64_Licensee">
