Reduce the TSAN EditorPerf timeouts and increase logging level temoprarily while we diagnose issues on the farm

#rnx

diff --git a/Engine/Build/EditorPerf.xml b/Engine/Build/EditorPerf.xml
--- a/Engine/Build/EditorPerf.xml
+++ b/Engine/Build/EditorPerf.xml
@@ -189,7 +189,7 @@
 	<Property Name="GauntletArgs" Value="-build=editor -tempdir=&quot;$(RootDir)/Tests&quot;"/>
 	<Property Name="GauntletArgs" Value="$(GauntletArgs) -branch=$(Branch) -ECBranch=&quot;$(Branch)&quot; -ECChangelist=&quot;$(Change)&quot;"/>
 	<Property Name="GauntletArgs" Value="$(GauntletArgs) -verbose -maxduration=7200" If="!$(WithTSAN)"/>
-	<Property Name="GauntletArgs" Value="$(GauntletArgs) -verbose -maxduration=18000 -LogIdleTimeout=14400" If="$(WithTSAN)"/>
+	<Property Name="GauntletArgs" Value="$(GauntletArgs) -VeryVerbose -maxduration=7200 -LogIdleTimeout=1800" If="$(WithTSAN)"/>
 
 	<!-- Gauntlet arguments for profile runs -->
 	<Property Name="ProfileArgs" Value="-tracefiletrunc"/>
