Fixing CreateArtifact BG Task attributes

#rnx

diff --git a/Engine/Build/InstalledEngineBuild.xml b/Engine/Build/InstalledEngineBuild.xml
--- a/Engine/Build/InstalledEngineBuild.xml
+++ b/Engine/Build/InstalledEngineBuild.xml
@@ -1169,7 +1169,7 @@
 			<SanitizeReceipt Files="#Installed Build Win64 Receipts" EngineDir="$(LocalInstalledDir)\Engine"/>
 		</Node>
 		<Node Name="Upload To Horde Win64" Requires="Make Installed Build Win64">
-			<CreateArtifact Name="installed-build-win64" Type="installed-build" Description="Installed Build (Windows)" BaseDir="$(LocalInstalledDir)"/>
+			<CreateArtifact Name="installed-build-win64" Type="installed-build" Description="Installed Build (Windows)" BaseDir="$(LocalInstalledDir)" Files="..."/>
 		</Node>
 	</Agent>
 	<Aggregate Name="HostPlatforms_Win64" Label="Builds/Win64" Requires="Make Installed Build Win64" Exclude="Update Version Files"/>
@@ -1296,7 +1296,7 @@
 			<SanitizeReceipt Files="#Installed Build Mac Receipts" EngineDir="$(LocalInstalledDirMac)/Engine"/>
 		</Node>
 		<Node Name="Upload To Horde Mac" Requires="Make Installed Build Mac">
-			<CreateArtifact Name="installed-build-mac" Type="installed-build" Description="Installed Build (Mac)" BaseDir="$(LocalInstalledDirMac)"/>
+			<CreateArtifact Name="installed-build-mac" Type="installed-build" Description="Installed Build (Mac)" BaseDir="$(LocalInstalledDirMac)" Files="..."/>
 		</Node>
 	</Agent>
 	<Aggregate Name="HostPlatforms_Mac" Label="Builds/Mac" Requires="Make Installed Build Mac"/>
@@ -1414,7 +1414,7 @@
 			<SanitizeReceipt Files="#Installed Build Linux Receipts" EngineDir="$(LocalInstalledDirLinux)/Engine"/>
 		</Node>
 		<Node Name="Upload To Horde Linux" Requires="Make Installed Build Linux">
-			<CreateArtifact Name="installed-build-linux" Type="installed-build" Description="Installed Build (Linux)" BaseDir="$(LocalInstalledDirLinux)"/>
+			<CreateArtifact Name="installed-build-linux" Type="installed-build" Description="Installed Build (Linux)" BaseDir="$(LocalInstalledDirLinux)" Files="..."/>
 		</Node>
 	</Agent>
 	<Aggregate Name="HostPlatforms_Linux" Label="Builds/Linux" Requires="Make Installed Build Linux"/>
