Enable verbose logging for warm TSAN jobs only so we can unblock cold runs to diagnose the sporadically failing warm runs

#rnx

diff --git a/Engine/Build/EditorPerf.xml b/Engine/Build/EditorPerf.xml
--- a/Engine/Build/EditorPerf.xml
+++ b/Engine/Build/EditorPerf.xml
@@ -189,7 +189,7 @@
 	<Property Name="GauntletArgs" Value="-build=editor -tempdir=&quot;$(RootDir)/Tests&quot;"/>
 	<Property Name="GauntletArgs" Value="$(GauntletArgs) -branch=$(Branch) -ECBranch=&quot;$(Branch)&quot; -ECChangelist=&quot;$(Change)&quot;"/>
 	<Property Name="GauntletArgs" Value="$(GauntletArgs) -verbose -maxduration=7200" If="!$(WithTSAN)"/>
-	<Property Name="GauntletArgs" Value="$(GauntletArgs) -VeryVerbose -maxduration=7200 -LogIdleTimeout=1800" If="$(WithTSAN)"/>
+	<Property Name="GauntletArgs" Value="$(GauntletArgs) -maxduration=7200 -LogIdleTimeout=1800" If="$(WithTSAN)"/>
 
 	<!-- Gauntlet arguments for profile runs -->
 	<Property Name="ProfileArgs" Value="-tracefiletrunc"/>
@@ -254,6 +254,7 @@
 		<Property Name="ReportParams" Value="-WriteTestResultsForHorde -HordeTestDataKey=&quot;$(NodeName)$(PreflightSuffix)&quot;"/>
 		<Property Name="ArtifactName" Value="$(Phase)$(IterationSuffix)$(TestNameSuffix)"/>
 		<Property Name="AllGauntletArgs" Value="-Project=$(ProjectTarget) -Platform=$(Platform) $(GauntletArgs) $($(ProjectName)__GauntletTestArgs) $(ReportParams) $(LogOutputParams) -artifactname=&quot;$(ArtifactName)&quot;"/>
+		<Property Name="AllGauntletArgs" Value="-Project=$(ProjectTarget) -Platform=$(Platform) $(GauntletArgs) -VeryVerbose $($(ProjectName)__GauntletTestArgs) $(ReportParams) $(LogOutputParams) -artifactname=&quot;$(ArtifactName)&quot;" If="$(WithTSAN) and ('$(Phase)' == 'Warm')"/>
 
 		<Property Name="TestArg" Value="-test=$($(ProjectName)__GauntletTest)" />
 		<Property Name="RunArgs" Value="$(TestArg) $(TestFilterArgs) $(Args) -editorargs=&quot;$(CombinedEditorArgs)&quot;"/>
