Log stdout when running TSAN to ensure we are logging as much as possible until we can determine the cause of sporadic Engine Initialization failures

#rnx

diff --git a/Engine/Build/EditorPerf.xml b/Engine/Build/EditorPerf.xml
--- a/Engine/Build/EditorPerf.xml
+++ b/Engine/Build/EditorPerf.xml
@@ -254,7 +254,7 @@
 		<Property Name="ReportParams" Value="-WriteTestResultsForHorde -HordeTestDataKey=&quot;$(NodeName)$(PreflightSuffix)&quot;"/>
 		<Property Name="ArtifactName" Value="$(Phase)$(IterationSuffix)$(TestNameSuffix)"/>
 		<Property Name="AllGauntletArgs" Value="-Project=$(ProjectTarget) -Platform=$(Platform) $(GauntletArgs) $($(ProjectName)__GauntletTestArgs) $(ReportParams) $(LogOutputParams) -artifactname=&quot;$(ArtifactName)&quot;"/>
-		<Property Name="AllGauntletArgs" Value="-Project=$(ProjectTarget) -Platform=$(Platform) $(GauntletArgs) -VeryVerbose $($(ProjectName)__GauntletTestArgs) $(ReportParams) $(LogOutputParams) -artifactname=&quot;$(ArtifactName)&quot;" If="$(WithTSAN) and ('$(Phase)' == 'Warm')"/>
+		<Property Name="AllGauntletArgs" Value="-Project=$(ProjectTarget) -Platform=$(Platform) $(GauntletArgs) -stdout -VeryVerbose $($(ProjectName)__GauntletTestArgs) $(ReportParams) $(LogOutputParams) -artifactname=&quot;$(ArtifactName)&quot;" If="$(WithTSAN) and ('$(Phase)' == 'Warm')"/>
 
 		<Property Name="TestArg" Value="-test=$($(ProjectName)__GauntletTest)" />
 		<Property Name="RunArgs" Value="$(TestArg) $(TestFilterArgs) $(Args) -editorargs=&quot;$(CombinedEditorArgs)&quot;"/>
