Add -uba to the compile portions of Low Level Tests when the option is enabled for alternative code path
#rnx

diff --git a/Engine/Build/LowLevelTests.xml b/Engine/Build/LowLevelTests.xml
--- a/Engine/Build/LowLevelTests.xml
+++ b/Engine/Build/LowLevelTests.xml
@@ -148,6 +148,7 @@
 			<Command Name="BuildCookRun" Arguments="-Build -NoBootstrapExe -NoCodeSign -SkipCook -SkipBuildEditor=True -Stage -NoSubmit -Archive -ArchiveDirectory=$(LocalAchiveDirectory) -project=$(TargetName) -configuration=$(Configuration) -platform=$(Platform) $(ExtraCompilationArgsFinal)" If="!$(DryRun)"/>
 		</Do>
 		<Do If="'$(StagesWithProjectFile)' != 'True'">
+		<Property Name="ExtraCompilationArgsFinal" Value="$(ExtraCompilationArgsFinal) -UBA" If="$(WithUBA)"/>
 			<Do If="'$(BuildProject)' != ''">
 				<Property Name="ExtraCompilationArgsFinal" Value="$(ExtraCompilationArgs) -Project=$(RootDir)/$(BuildProject)"/>
 			</Do>
