Convert installed engine build to use LiveLinkHubEditor target rather than the LiveLinkHub target.

Moved LiveLinkHub to the Compile Unreal Editor node in order to avoid generating the same binaries in different steps.
This should also reduce LLH's compile time and footprint since it will reuse Unreal's binaries because of thier shared modular build environment.

LLH Mac and Linux compile steps are currently disabled until its post-build steps are fixed on these platforms.

diff --git a/Engine/Build/InstalledEngineBuild.xml b/Engine/Build/InstalledEngineBuild.xml
--- a/Engine/Build/InstalledEngineBuild.xml
+++ b/Engine/Build/InstalledEngineBuild.xml
@@ -115,12 +115,6 @@
 	
 	<Property Name="DefaultArchitectureCommandsWin64" Value=""/>
 	
-	<!-- LiveLinkHub args -->
-	<Option Name="WithLiveLinkHub3rdPartyPlugins" DefaultValue="true" Restrict="true|false" Description="Whether or not to compile LiveLinkHub plugins."/>
-	<Option Name="HasExtraNotForLicenseeFiles" DefaultValue="false" Restrict="true|false" Description="Whether or not a an extra set of nodes should be generated."/>
-	<Property Name="LiveLinkHubCompileArgs" Value="-EnableThirdPartyPlugins=true" />
-	<Property Name="LiveLinkHubCompileArgs" Value="-EnableThirdPartyPlugins=false" If="!$(WithLiveLinkHub3rdPartyPlugins)"/>
-	
 	<!-- Compile flags for Windows Game binaries -->
 	<Property Name="TargetDebugInfoArg" Value="-nodebuginfo"/>
 	<Property Name="TargetDebugInfoArg" Value="" If="$(WithFullDebugInfo)"/>
@@ -197,6 +191,9 @@
 		<Node Name="Compile UnrealEditor Win64" Requires="Update Version Files" Produces="#UnrealEditor Win64;#UnrealEditor Win64 Unstripped;#UnrealEditor Win64 Stripped;#UnrealEditor Win64 Unsigned;#UnrealEditor Win64 Signed;#UnrealEditor Win64 WithExceptions">
 			<Compile Target="UnrealEditor" Platform="Win64" Configuration="DebugGame" Tag="#UnrealEditor Win64" Arguments="-precompile -allmodules $(ExtraCompileArgs) $(Win64NoPCHArg)"/>
 			<Compile Target="UnrealEditor" Platform="Win64" Configuration="Development" Tag="#UnrealEditor Win64" Arguments="-precompile -allmodules $(ExtraCompileArgs) $(Win64NoPCHArg)"/>
+			
+			<!-- Note: We compile LiveLinkHubEditor here because we share binaries with the editor. This drastically reduces its build time and prevents having duplicate files in the MakeInstalledBuild step. -->
+			<Compile Target="LiveLinkHubEditor" Configuration="Development" Platform="Win64" Tag="#UnrealEditor Win64" Arguments="-precompile $(ExtraCompileArgs) $(Win64NoPCHArg)" />
 
 			<Do If="$(EmbedSrcSrvInfo)">
 				<!-- Embed source info into the PDB files. Should be done from this machine to ensure that paths are correct. -->
@@ -235,8 +232,12 @@
 	<Agent Name="Editor Mac" Type="Mac_Licensee">
 		<Node Name="Compile UnrealEditor Mac" Requires="Update Version Files" Produces="#UnrealEditor Mac;#UnrealEditor Mac Unstripped;#UnrealEditor Mac Stripped;#UnrealEditor Mac Unsigned;#UnrealEditor Mac Signed;#UnrealEditor Mac WithExceptions">
 			<Property Name="AllModulesArg" Value="-allmodules"/>
+
 			<Compile Target="UnrealEditor" Platform="Mac" Configuration="DebugGame" Tag="#UnrealEditor Mac" Arguments="-precompile -EnableDSYM $(ExtraCompileArgsMac) $(AllModulesArg) $(MacArchitectureArg)"/>
 			<Compile Target="UnrealEditor" Platform="Mac" Configuration="Development" Tag="#UnrealEditor Mac" Arguments="-precompile -EnableDSYM $(ExtraCompileArgsMac) $(AllModulesArg) $(MacArchitectureArg)"/>
+			
+			<!-- Note: We compile LiveLinkHubEditor here because we share binaries with the editor. This drastically reduces its build time and prevents having duplicate files in it -->
+			<!--<Compile Target="LiveLinkHubEditor" Configuration="Development" Platform="Mac" Tag="#UnrealEditor Mac" Arguments="-precompile -EnableDSYM $(MacArchitectureArg)" /> -->
 
 			<Property Name="UnrealEditorExceptions" Value=""/>
 			<Property Name="UnrealEditorExceptions" Value="$(PluginsExceptions)" If="$(CompileDatasmithPlugins)"/>
@@ -269,6 +270,9 @@
 		<Node Name="Compile UnrealEditor Linux" Requires="Update Version Files" Produces="#UnrealEditor Linux;#UnrealEditor Linux Unstripped;#UnrealEditor Linux Stripped;#UnrealEditor Linux WithExceptions">
 			<Compile Target="UnrealEditor" Platform="Linux" Configuration="DebugGame" Tag="#UnrealEditor Linux" Arguments="-precompile -allmodules $(ExtraCompileArgs)"/>
 			<Compile Target="UnrealEditor" Platform="Linux" Configuration="Development" Tag="#UnrealEditor Linux" Arguments="-precompile -allmodules $(ExtraCompileArgs)"/>
+			
+			<!-- Note: We compile LiveLinkHubEditor here because we share binaries with the editor. This drastically reduces its build time and prevents having duplicate files in the MakeInstalledBuild step. -->
+			<!--<Compile Target="LiveLinkHubEditor" Configuration="Development" Platform="Linux" Tag="#UnrealEditor Linux" Arguments="-precompile $(ExtraCompileArgs)" />-->
 
 			<Tag Files="#UnrealEditor Linux" Filter="$(LinuxStripFilter)" With="#UnrealEditor Linux Unstripped"/>
 			<Tag Files="#UnrealEditor Linux" Except="$(PluginsExceptions)" With="#UnrealEditor Linux WithExceptions"/>
@@ -478,54 +482,15 @@
 	</Agent>
 	<Aggregate Name="TargetPlatforms_TVOS" Label="Platforms/TVOS" Requires="Compile UnrealGame TVOS"/>
 
-	<Do If="!$(HasExtraNotForLicenseeFiles)">
-		<!-- If this is false we stub out the macro. -->
-		<Macro Name="ExtraNotForLicenseeFiles" Arguments="TagSet;Platform;HostPlatform" />
-	</Do>
-
 	<Agent Name="NotForLicensee Tools Group Win64" Type="Win64">
 		<Node Name="Build NotForLicensee Tools Win64" Requires="Update Version Files" Produces="#NotForLicensee Build Tools Win64">
 			<Compile Target="DatasmithCADWorker" Configuration="Development" Platform="Win64" Tag="#NotForLicensee Build Tools Win64" Arguments="$(ExtraCompileArgs)"/>
 			<!-- Temporarily removing the compilation of DatasmithCADRuntime from the installed build
 			<Compile Target="DatasmithCADRuntime" Configuration="Development" Platform="Win64" Tag="#NotForLicensee Build Tools Win64" Arguments="$(ExtraCompileArgs)"/>
 			-->
-			<Compile Target="LiveLinkHub" Configuration="Shipping" Platform="Win64" Arguments="-Monolithic $(LiveLinkHubCompileArgs) $(ExtraCompileArgs)" />
-
-			<Expand Name="ExtraNotForLicenseeFiles" TagSet="#NotForLicensee Build Tools Win64" Platform="Win64" HostPlatform="Win64"/>
-			<Tag Files="Engine/Binaries/Win64/LiveLinkHub-Win64-Shipping.exe" With="#NotForLicensee Build Tools Win64"/>
-			<Tag Files="Engine/Programs/LiveLinkHub/..." With="#NotForLicensee Build Tools Win64" />
-		</Node>
-		
-		<Node Name="Build NotForLicensee Tools Linux" Requires="Update Version Files" Produces="#NotForLicensee Build Tools Linux">
-			<Compile Target="LiveLinkHub" Configuration="Shipping" Platform="Linux" Arguments="-Monolithic $(LiveLinkHubCompileArgs) $(ExtraCompileArgs)"/>
-
-			<!-- Since we're running on the same agent as the Win64 node, we don't need to re-tag those files or else it would generate an error. -->
-			<Expand Name="ExtraNotForLicenseeFiles" TagSet="#NotForLicensee Build Tools Linux" Platform="Linux" HostPlatform="Linux"/>
-			<Tag Files="Engine/Binaries/Linux/LiveLinkHub-Linux-Shipping" With="#NotForLicensee Build Tools Linux" />
-		</Node>
-	</Agent>
-	
-	<Agent Name="NotForLicensee Tools Group LinuxNative" Type="Linux">
-		<Node Name="Build NotForLicensee Tools LinuxNative" Requires="Update Version Files" Produces="#NotForLicensee Build Tools LinuxNative">
-			<Compile Target="LiveLinkHub" Configuration="Shipping" Platform="Linux" Arguments="-Monolithic $(LiveLinkHubCompileArgs) $(ExtraCompileArgs)"/>
-
-			<Expand Name="ExtraNotForLicenseeFiles" TagSet="#NotForLicensee Build Tools LinuxNative" Platform="LinuxNative" HostPlatform="Linux"/>
-			<Tag Files="Engine/Binaries/Linux/LiveLinkHub-Linux-Shipping" With="#NotForLicensee Build Tools LinuxNative" />
-			<Tag Files="Engine/Programs/LiveLinkHub/..." With="#NotForLicensee Build Tools LinuxNative" />
 		</Node>
 	</Agent>
 	
-	<Agent Name="NotForLicensee Tools Group Mac" Type="Mac">
-		<Node Name="Build NotForLicensee Tools Mac" Requires="Update Version Files" Produces="#NotForLicensee Build Tools Mac">
-			<Compile Target="LiveLinkHub" Configuration="Shipping" Platform="Mac" Arguments="-Monolithic $(LiveLinkHubCompileArgs) $(ExtraCompileArgsMac)"/>
-
-			<Expand Name="ExtraNotForLicenseeFiles" TagSet="#NotForLicensee Build Tools Mac" Platform="Mac" HostPlatform="Mac"/>
-			<Tag Files="Engine/Binaries/Mac/LiveLinkHub-Mac-Shipping.app/..." With="#NotForLicensee Build Tools Mac" />
-			<Tag Files="Engine/Programs/LiveLinkHub/..." With="#NotForLicensee Build Tools Mac" />
-		</Node>
-	</Agent>
-
-	
 	<Agent Name="BuildRules" Type="Win64_Licensee">
 		<Node Name="Compile Rules Assemblies" Produces="#Rules Assemblies">
 			<Property Name="UBTPath" Value="Engine/Binaries/DotNET/UnrealBuildTool/UnrealBuildTool.exe"/>
@@ -622,8 +587,7 @@
 			<!-- Sign the binaries -->
 			<Sign Files="#Sign Binaries" Parallel="$(SignWindowsExecutablesInParallel)" If="'$(SignExecutables)' == true"/>
 		</Node>
-		<Node Name="Build Tools Linux" Requires="Update Version Files;#NotForLicensee Build Tools Linux" >
-			<Tag Files="#NotForLicensee Build Tools Linux" With="#Build Tools Linux"/>
+		<Node Name="Build Tools Linux" Requires="Update Version Files" >
 			<Compile Target="CrashReportClient" Platform="Linux" Configuration="Shipping" Tag="#Build Tools Linux" Arguments="$(CRCTelemetryDefines)  $(CRCDefines) $(ExtraCompileArgs)"/>
 			<Compile Target="CrashReportClientEditor" Platform="Linux" Configuration="Shipping" Tag="#Build Tools Linux" Arguments="$(CRCETelemetryDefines)  $(CRCEDefines) $(ExtraCompileArgs)"/>
 			<Compile Target="ShaderCompileWorker" Configuration="Development" Platform="Linux" Tag="#Build Tools Linux" Arguments="$(ExtraCompileArgs)"/>
@@ -693,8 +657,7 @@
 			<Sign Files="#Sign Binaries Mac" If="'$(SignExecutables)' == true"/>
 		</Node>
 	
-		<Node Name="Build Tools Mac" Requires="Update Version Files;#NotForLicensee Build Tools Mac">
-			<Tag Files="#NotForLicensee Build Tools Mac" With="#Build Tools Mac"/>
+		<Node Name="Build Tools Mac" Requires="Update Version Files">
 			<Compile Target="CrashReportClient" Configuration="Shipping" Platform="Mac" Tag="#Build Tools Mac" Arguments="-EnableDSYM $(CRCTelemetryDefines) $(CRCDefines) $(ExtraCompileArgsMac)"/>
 			<Compile Target="CrashReportClientEditor" Configuration="Shipping" Platform="Mac" Tag="#Build Tools Mac" Arguments="-EnableDSYM $(CRCETelemetryDefines) $(CRCEDefines) $(ExtraCompileArgsMac)"/>
 			<Compile Target="ShaderCompileWorker" Configuration="Development" Platform="Mac" Tag="#Build Tools Mac" Arguments="-EnableDSYM $(ExtraCompileArgsMac)"/>
@@ -717,8 +680,7 @@
 	</Agent>
 
 	<Agent Name="Tools Group LinuxNative" Type="Linux_Licensee">
-		<Node Name="Build Tools LinuxNative" Requires="Update Version Files;#NotForLicensee Build Tools LinuxNative">
-			<Tag Files="#NotForLicensee Build Tools LinuxNative" With="#Build Tools LinuxNative"/>
+		<Node Name="Build Tools LinuxNative" Requires="Update Version Files">
 			<Compile Target="CrashReportClient" Platform="Linux" Configuration="Shipping" Tag="#Build Tools LinuxNative" Arguments="$(CRCTelemetryDefines) $(CRCDefines) $(ExtraCompileArgs)"/>
 			<Compile Target="CrashReportClientEditor" Platform="Linux" Configuration="Shipping"  Tag="#Build Tools LinuxNative" Arguments="$(CRCETelemetryDefines) $(CRCEDefines) $(ExtraCompileArgs)"/>
 			<Compile Target="ShaderCompileWorker" Configuration="Development" Platform="Linux" Tag="#Build Tools LinuxNative" Arguments="$(ExtraCompileArgs)"/>
diff --git a/Engine/Build/InstalledEngineFilters.xml b/Engine/Build/InstalledEngineFilters.xml
--- a/Engine/Build/InstalledEngineFilters.xml
+++ b/Engine/Build/InstalledEngineFilters.xml
@@ -347,10 +347,6 @@
 	<Property Name="CopyEditorExceptionsWin64" Value="$(CopyEditorExceptions)">
 		<!-- Exclude Mac binaries on windows here, because previous wildcard for plugin resources matches OSX dsym resources on Windows -->
 		Engine/Plugins/.../Binaries/Mac/...
-		
-		Engine/Binaries/Win64/LiveLinkHub-Win64-Shipping.pdb
-		Engine/Binaries/Linux/LiveLinkHub-Linux-Shipping.debug
-		Engine/Binaries/Linux/LiveLinkHub-Linux-Shipping.sym
 	</Property>
 
 	<!-- Mac editor specific dependencies -->
@@ -496,10 +492,6 @@
 		<!-- Don't want The Windows/Mac Python packages -->
 		Engine/Plugins/Experimental/PythonFoundationPackages/Content/Python/Lib/Win64/...
 		Engine/Plugins/Experimental/PythonFoundationPackages/Content/Python/Lib/Mac/...
-		
-		<!-- We remove these in an attempt to save space. -->
-		Engine/Binaries/Linux/LiveLinkHub-Linux-Shipping.debug
-		Engine/Binaries/Linux/LiveLinkHub-Linux-Shipping.sym
 	</Property>
 
 	<!-- Target Platform Filters/Exceptions -->
