- VisionOS in PrecompiledBinary builds

diff --git a/Engine/Build/InstalledEngineBuild.xml b/Engine/Build/InstalledEngineBuild.xml
--- a/Engine/Build/InstalledEngineBuild.xml
+++ b/Engine/Build/InstalledEngineBuild.xml
@@ -144,10 +144,14 @@
 	<Macro Name="Platform_FilterRestrictedFolders"/>
 	<Macro Name="Platform_FilterOverrides"/>
 	<Macro Name="Platform_Agent_Target_Win64_Licensee"/>
+	<Macro Name="Platform_Agent_Target_Mac_Licensee"/>
+	<Macro Name="Platform_Agent_Target_Linux_Licensee"/>
 	<Macro Name="Platform_AddInstalledRequirements"/>
+	<Macro Name="Platform_AddFilterRequirements"/>
 	<Macro Name="Platform_TagCsToolsRedistUATSource"/>
 	<Macro Name="Platform_Build_Tools_CS_Binaries"/>
 	<Macro Name="Platform_Append_DDCPlatformsWin64"/>
+	<Macro Name="Platform_Append_DDCPlatformsMac"/>
 	<Macro Name="Platform_TagDDCFiles"/>
 	<Macro Name="Platform_Tag_UnrealGame"/>
 	<Macro Name="Platform_ExtraSignFiles"/>
@@ -371,6 +375,8 @@
 	</Agent>
 	<Aggregate Name="TargetPlatforms_LinuxNative" Requires="Compile UnrealGame LinuxNative"/>
 
+	<Expand Name="Platform_Agent_Target_Linux_Licensee"/>
+
 	<Agent Name="Target LinuxArm64" Type="Win64_Licensee">
 		<Node Name="Compile UnrealGame LinuxArm64" Requires="Update Version Files" Produces="#UnrealGame LinuxArm64;#UnrealGame LinuxArm64 Unstripped;#UnrealGame LinuxArm64 Stripped">
 			<ForEach Name="Target" Values="UnrealGame;$(OptionalClientTarget);$(OptionalServerTarget)">
@@ -482,6 +488,8 @@
 	</Agent>
 	<Aggregate Name="TargetPlatforms_TVOS" Label="Platforms/TVOS" Requires="Compile UnrealGame TVOS"/>
 
+	<Expand Name="Platform_Agent_Target_Mac_Licensee"/>
+
 	<Agent Name="NotForLicensee Tools Group Win64" Type="Win64">
 		<Node Name="Build NotForLicensee Tools Win64" Requires="Update Version Files" Produces="#NotForLicensee Build Tools Win64">
 			<Compile Target="DatasmithCADWorker" Configuration="Development" Platform="Win64" Tag="#NotForLicensee Build Tools Win64" Arguments="$(ExtraCompileArgs)"/>
@@ -844,7 +852,7 @@
 			<Tag Files="Engine\Binaries\Win64\zenserver.*" With="#ToCopy"/>
 			<Tag Files="Engine\Binaries\Win64\UnrealBuildAccelerator\..." With="#ToCopy"/>
 			<Tag Files="Templates\TemplateResources\..." With="#ToCopy"/>
-			<Expand Name="TagDDCEngineFiles" ExceptPlugins="Plugins\...\Binaries\Mac\..."/>
+			<Expand Name="TagDDCEngineFiles" ExceptPlugins="Plugins/.../Binaries/Mac/..."/>
 			<!-- Filter out the files not needed to build DDC. Removing confidential folders can affect DDC keys, so we want to be sure that we're making DDC with a build that can use it. -->
 			<Tag Files="#ToCopy" Except="$(ConfidentialExceptions);$(RootDir)/.../Source/...;$(RootDir)/.../Intermediate/..." With="#FilteredCopyList"/>
 			
@@ -873,6 +881,7 @@
 		<Property Name="DDCPlatformsMac" Value="$(DDCPlatformsMac)+Android_ATC" If="'$(WithAndroid)' == true"/>
 		<Property Name="DDCPlatformsMac" Value="$(DDCPlatformsMac)+IOS" If="'$(WithIOS)' == true"/>
 		<Property Name="DDCPlatformsMac" Value="$(DDCPlatformsMac)+TVOS" If="'$(WithTVOS)' == true"/>
+		<Expand Name="Platform_Append_DDCPlatformsMac"/>
 	</Do>
 
 	<Property Name="MacDDCAgentType" Value="Mac" />
@@ -1150,6 +1159,8 @@
 		<Property Name="FilterRequirements" Value="$(FilterRequirements);#CompressedDDCMac" If="$(WithDDC) And $(WithMac)"/>
 		<Property Name="FilterRequirements" Value="$(FilterRequirements);#Compile Datasmith Plugins Mac;#DatasmithPluginBinaries Mac" If="$(WithMac) And $(CompileDatasmithPlugins)"/>
 
+		<Expand Name="Platform_AddFilterRequirements" />
+
 		<!-- Feature packs -->
 		<Property Name="FilterRequirements" Value="$(FilterRequirements);Make Feature Packs Mac"/>
 
@@ -1200,6 +1211,8 @@
 				<Property Name="CopyInstalledFilter" Value="$(CopyInstalledFilter);$(CopyTVOSFilter)"/>
 				<Property Name="CopyInstalledExceptions" Value="$(CopyInstalledExceptions);$(CopyTVOSExceptions)"/>
 			</Do>
+			<Expand Name="Platform_Tag_UnrealGame" />
+			<Expand Name="Platform_ExtraSignFiles" />
 
 			<!-- Tag any dependencies from all previous build products -->
 			<Tag Files="#Installed Mac" Filter="*.target" With="#TargetReceipts"/>
@@ -1320,6 +1333,8 @@
 				<Property Name="CopyInstalledFilter" Value="$(CopyInstalledFilter);$(CopyAndroidFilter)"/>
 				<Property Name="CopyInstalledExceptions" Value="$(CopyInstalledExceptions);$(CopyAndroidExceptions)"/>
 			</Do>
+			<Expand Name="Platform_Tag_UnrealGame" />
+			<Expand Name="Platform_ExtraSignFiles"/>
 
 			<!-- Tag any dependencies from all previous build products -->
 			<Tag Files="#Installed Linux" Filter="*.target" With="#TargetReceipts"/>
diff --git a/Engine/Platforms/VisionOS/Build/InstalledEngineBuild.xml b/Engine/Platforms/VisionOS/Build/InstalledEngineBuild.xml
--- a/Engine/Platforms/VisionOS/Build/InstalledEngineBuild.xml
+++ b/Engine/Platforms/VisionOS/Build/InstalledEngineBuild.xml
@@ -3,8 +3,7 @@
 
 	<Option Name="WithVisionOS" Restrict="true|false" DefaultValue="false" Description="Include the VisionOS target platform"/>
 
-	<Property Name="ExtensionPlatforms" Value="$(ExtensionPlatforms);" If="'$(ExtensionPlatforms)' != ''"/>
-	<Property Name="ExtensionPlatforms" Value="$(ExtensionPlatforms)VisionOS"/>
+	<Property Name="ExtensionPlatforms" Value="$(ExtensionPlatforms);" If="'$(ExtensionPlatforms)' != ''"/>	<Property Name="ExtensionPlatforms" Value="$(ExtensionPlatforms)VisionOS"/>
 
 	<Do If="$(WithVisionOS)">
 		<Expand Name="AddExtensionPlatform" PlatformDir="VisionOS"/>
@@ -19,14 +18,42 @@
 			</Property>
 		</Do>
 	</Extend>
+	
+	<Extend Name="Platform_Agent_Target_Mac_Licensee">
+		<Agent Name="Target VisionOS" Type="Mac_Licensee">
+			<Node Name="Compile UnrealGame VisionOS" Requires="Update Version Files" Produces="#UnrealGame VisionOS">
+				<ForEach Name="Target" Values="UnrealGame;$(OptionalClientTarget)">
+					<ForEach Name="Configuration" Values="$(GameConfigurations)">
+						<Compile Target="$(Target)" Platform="VisionOS" Configuration="$(Configuration)" Tag="#UnrealGame VisionOS" AllowParallelExecutor="$(bAllowParallelExecutor)" Arguments="-precompile -allmodules -nolink $(TargetDebugInfoArg) $(ExtraCompileArgs)"/>
+						<Compile Target="$(Target)" Platform="VisionOS" Configuration="$(Configuration)" Tag="#UnrealGame VisionOS" AllowParallelExecutor="$(bAllowParallelExecutor)" Arguments="-precompile $(TargetDebugInfoArg) $(ExtraCompileArgs)" Clean="false"/>
+					</ForEach>
+				</ForEach>
+				
+				<!-- sharing IOS's strip filter since it's in a InstalledEngineFilters.xml, and we don't want to put VisionOS in there -->
+				<Tag Files="#UnrealGame VisionOS" Filter="$(IOSStripFilter)" With="#UnrealGame VisionOS Unstripped"/>
+				<Strip BaseDir="$(RootDir)" OutputDir="$(SavedOutput)" Platform="VisionOS" Files="#UnrealGame VisionOS Unstripped" Tag="#UnrealGame VisionOS Stripped"/>
+			</Node>
+		</Agent>
+		<Aggregate Name="TargetPlatforms_VisionOS" Label="Target Platforms/VisionOS" Requires="Compile UnrealGame VisionOS"/>
+	</Extend>
 
+	<!-- Adds to 'Installed Build Group Win64' -->
+	<!--
 	<Extend Name="Platform_AddInstalledRequirements">
 		<Do If="$(WithVisionOS)">
 			<Property Name="InstalledRequirements" Value="$(InstalledRequirements);#UnrealGame VisionOS"/>
 		</Do>
 	</Extend>
+	-->
+
+	<!-- Adds to 'Installed Build Group Mac' -->
+	<Extend Name="Platform_AddFilterRequirements">
+		<Do If="$(WithVisionOS)">
+			<Property Name="FilterRequirements" Value="$(FilterRequirements);#UnrealGame VisionOS"/>
+		</Do>
+	</Extend>
 
-	<Extend Name="Platform_Append_DDCPlatformsWin64">
+	<Extend Name="Platform_Append_DDCPlatformsMac">
 		<Do If="$(WithVisionOS)">
 			<Property Name="DDCPlatformsMac" Value="$(DDCPlatformsMac)+VisionOS"/>
 		</Do>
@@ -51,7 +78,8 @@
 			</Property>
 			<Property Name="CopyVisionOSExceptions">
 			</Property>
-			<Tag Files="#UnrealGame VisionOS" With="#Installed Mac"/>
+			<Tag Files="#UnrealGame VisionOS" Except="#UnrealGame VisionOS Unstripped" With="#Installed Mac"/>
+			<Tag Files="#UnrealGame VisionOS Stripped" With="#Saved Output"/>
 			<Property Name="Platforms" Value="$(Platforms)VisionOS;"/>
 			<Property Name="CopyInstalledFilter" Value="$(CopyInstalledFilter);$(CopyVisionOSFilter)"/>
 			<Property Name="CopyInstalledExceptions" Value="$(CopyInstalledExceptions);$(CopyVisionOSExceptions)"/>
