Include ushell in installed UE

#rnx

diff --git a/Engine/Build/InstalledEngineFilters.xml b/Engine/Build/InstalledEngineFilters.xml
--- a/Engine/Build/InstalledEngineFilters.xml
+++ b/Engine/Build/InstalledEngineFilters.xml
@@ -301,6 +301,7 @@
 		Engine/Extras/MayaLiveLink/...
 		Engine/Extras/MayaVelocityGridExporter/...
 		Engine/Extras/UnrealVS/...
+		Engine/Extras/ushell/...
 		Engine/Extras/Redist/en-us/*
 		Engine/Extras/VirtualProduction/TextureShare/...
 
