Disable LiveLinkHub compilation in installed build job to fix CIS error

diff --git a/Engine/Build/InstalledEngineBuild.xml b/Engine/Build/InstalledEngineBuild.xml
--- a/Engine/Build/InstalledEngineBuild.xml
+++ b/Engine/Build/InstalledEngineBuild.xml
@@ -197,7 +197,7 @@
 			<Compile Target="UnrealEditor" Platform="Win64" Configuration="Development" Tag="#UnrealEditor Win64" Arguments="-precompile -allmodules $(ExtraCompileArgs) $(Win64NoPCHArg)"/>
 			
 			<!-- Note: We compile LiveLinkHub here because we share a modular build environment with the editor. This drastically reduces its build time and prevents having duplicate files in the MakeInstalledBuild step. -->
-			<Compile Target="LiveLinkHub" Configuration="Development" Platform="Win64" Tag="#UnrealEditor Win64" Arguments="-precompile $(ExtraCompileArgs) $(Win64NoPCHArg)" />
+			<!--<Compile Target="LiveLinkHub" Configuration="Development" Platform="Win64" Tag="#UnrealEditor Win64" Arguments="-precompile $(ExtraCompileArgs) $(Win64NoPCHArg)" />-->
 
 			<Do If="$(EmbedSrcSrvInfo)">
 				<!-- Embed source info into the PDB files. Should be done from this machine to ensure that paths are correct. -->
