[CIS] Always dump shader timing stats to CSV file in EditorPerf builds.

This enables the newly added CVar "r.ShaderCompiler.DumpShaderTimeStats" (see CL 36574032) in the EditorPerf build configuration.

#rnx

diff --git a/Engine/Build/EditorPerf.xml b/Engine/Build/EditorPerf.xml
--- a/Engine/Build/EditorPerf.xml
+++ b/Engine/Build/EditorPerf.xml
@@ -93,6 +93,9 @@
 	</Do>
 	<Property Name="ArgIterationsPIE" Value="-ini:Engine:[ConsoleVariables]:project.iteration.pie.testiterations=$(NumIterationsPIE)" />
 
+	<!-- Arguments for dumping shader compilation timing statistics to CSV files as Horde artifacts -->
+	<Property Name="DumpShaderTimingStatsArgs" Value="-DPCVars=r.ShaderCompiler.DumpShaderTimeStats=true" />
+
 	<!-- Used to determine if we will include internal projects -->
 	<Option Name="IsLicensee" Restrict="true|false" DefaultValue="false" Description="Define to true if this script is used externally."/>
 
@@ -233,7 +236,7 @@
 		<Do If="'$($(ProjectName)__GauntletTest)'=='EditorTests'">
 			<Property Name="TestFilter" Value="Project.Iteration.PIE"/>
 			<Property Name="TestFilterArgs" Value="-runtest=$(TestFilter)"/>
-			<Property Name="CombinedEditorArgs" Value="$(EditorArgs) $(TestFilterArgs) $(AssetRegistryCacheRootFolderArg) $(ArgIterationsPIE) $(EditorTelemetryArgs)"/>
+			<Property Name="CombinedEditorArgs" Value="$(EditorArgs) $(TestFilterArgs) $(AssetRegistryCacheRootFolderArg) $(ArgIterationsPIE) $(EditorTelemetryArgs) $(DumpShaderTimingStatsArgs)"/>
 		</Do>
 		<Do If="'$($(ProjectName)__GauntletTest)'!='EditorTests'">
 			<Property Name="CombinedEditorArgs" Value="$(EditorArgs) $(AssetRegistryCacheRootFolderArg) $(EditorTelemetryArgs)"/>
