Fix missing files from LiveLinkHub installed build.

TargetInfo.json was not getting generated because UBT is getting invoked without the -project arg.
LiveLinkHub.target's copy in Source/Programs/LiveLinkHub was not getting copied since it was never tagged.

diff --git a/Engine/Build/InstalledEngineBuild.xml b/Engine/Build/InstalledEngineBuild.xml
--- a/Engine/Build/InstalledEngineBuild.xml
+++ b/Engine/Build/InstalledEngineBuild.xml
@@ -193,11 +193,15 @@
 			<SetVersion Change="$(Change)" Branch="$(EscapedBranch)" If="$(IsBuildMachine)"/>
 		</Node>
 		<Node Name="Compile UnrealEditor Win64" Requires="Update Version Files" Produces="#UnrealEditor Win64;#UnrealEditor Win64 Unstripped;#UnrealEditor Win64 Stripped;#UnrealEditor Win64 Unsigned;#UnrealEditor Win64 Signed;#UnrealEditor Win64 WithExceptions">
+			<!-- We invoke QueryTargets to generate LiveLinkHub's TargetInfo.json file which would otherwise be missing since we're invoking UBT without a -project arg. -->
+			<QueryTargets ProjectFile="Engine/Source/Programs/LiveLinkHub/LiveLinkHub.uproject" Tag="#UnrealEditor Win64"/>
+			
 			<Compile Target="UnrealEditor" Platform="Win64" Configuration="DebugGame" Tag="#UnrealEditor Win64" Arguments="-precompile -allmodules $(ExtraCompileArgs) $(Win64NoPCHArg)"/>
 			<Compile Target="UnrealEditor" Platform="Win64" Configuration="Development" Tag="#UnrealEditor Win64" Arguments="-precompile -allmodules $(ExtraCompileArgs) $(Win64NoPCHArg)"/>
 			
 			<!-- Note: We compile LiveLinkHub here because we share a modular build environment with the editor. This drastically reduces its build time and prevents having duplicate files in the MakeInstalledBuild step. -->
 			<Compile Target="LiveLinkHub" Configuration="Development" Platform="Win64" Tag="#UnrealEditor Win64" Arguments="-precompile $(ExtraCompileArgs) $(Win64NoPCHArg)" />
+			<Tag Files="Engine/Source/Programs/LiveLinkHub/Binaries/Win64/LiveLinkHub.target" With="#UnrealEditor Win64"/>
 
 			<Do If="$(EmbedSrcSrvInfo)">
 				<!-- Embed source info into the PDB files. Should be done from this machine to ensure that paths are correct. -->
diff --git a/Engine/Build/InstalledEngineFilters.xml b/Engine/Build/InstalledEngineFilters.xml
--- a/Engine/Build/InstalledEngineFilters.xml
+++ b/Engine/Build/InstalledEngineFilters.xml
@@ -182,6 +182,7 @@
 		Source/ThirdParty/Licenses/...
 		Source/Programs/AutomationTool/...
 		Source/Programs/AutomationToolLauncher/...
+		Source/Programs/LiveLinkHub/...
 		Source/Programs/Shared/...
 		Source/Programs/UnrealBuildTool/...
 		Source/Programs/UnrealHeaderTool/...
