[UBA]: Disable UBA for EditorPerf test on build machines.

This adds "-NoUbaController" to the command line arguments of the editor launch of EditorPerf.xml build script for two reasons:
1. Most users don't have UBA access in the WFH environemnt, so this test should have the same condition.
2. A performance issue with remaining jobs at the end of shader compilation on the build farm has been identified that drags out a handful of jobs extremely long, causing stalls at editor launch.
   Local shader compilation for this test on the build farm shows consistently better results and almost double the effective parallelization. We therefore disable UBA for this test until the issue has been resolved.

#rnx

diff --git a/Engine/Build/EditorPerf.xml b/Engine/Build/EditorPerf.xml
--- a/Engine/Build/EditorPerf.xml
+++ b/Engine/Build/EditorPerf.xml
@@ -236,7 +236,7 @@
 		<Do If="'$($(ProjectName)__GauntletTest)'=='EditorTests'">
 			<Property Name="TestFilter" Value="Project.Iteration.PIE"/>
 			<Property Name="TestFilterArgs" Value="-runtest=$(TestFilter)"/>
-			<Property Name="CombinedEditorArgs" Value="$(EditorArgs) $(TestFilterArgs) $(AssetRegistryCacheRootFolderArg) $(ArgIterationsPIE) $(EditorTelemetryArgs) $(DumpShaderTimingStatsArgs)"/>
+			<Property Name="CombinedEditorArgs" Value="$(EditorArgs) $(TestFilterArgs) $(AssetRegistryCacheRootFolderArg) $(ArgIterationsPIE) $(EditorTelemetryArgs) $(DumpShaderTimingStatsArgs) -NoUbaController"/>
 		</Do>
 		<Do If="'$($(ProjectName)__GauntletTest)'!='EditorTests'">
 			<Property Name="CombinedEditorArgs" Value="$(EditorArgs) $(AssetRegistryCacheRootFolderArg) $(EditorTelemetryArgs)"/>
