Re-enable LiveLinkHub build on Linux.

Keeping mac build disabled for now because of an error that is still being investigated.

diff --git a/Engine/Build/InstalledEngineBuild.xml b/Engine/Build/InstalledEngineBuild.xml
--- a/Engine/Build/InstalledEngineBuild.xml
+++ b/Engine/Build/InstalledEngineBuild.xml
@@ -240,12 +240,16 @@
 	<Agent Name="Editor Mac" Type="Mac_Licensee">
 		<Node Name="Compile UnrealEditor Mac" Requires="Update Version Files" Produces="#UnrealEditor Mac;#UnrealEditor Mac Unstripped;#UnrealEditor Mac Stripped;#UnrealEditor Mac Unsigned;#UnrealEditor Mac Signed;#UnrealEditor Mac WithExceptions">
 			<Property Name="AllModulesArg" Value="-allmodules"/>
+			
+			<!-- We invoke QueryTargets to generate LiveLinkHub's TargetInfo.json file which would otherwise be missing since we're invoking UBT without a -project arg. -->
+			<!--<QueryTargets ProjectFile="Engine/Source/Programs/LiveLinkHub/LiveLinkHub.uproject" Tag="#UnrealEditor Mac"/>-->
 
 			<Compile Target="UnrealEditor" Platform="Mac" Configuration="DebugGame" Tag="#UnrealEditor Mac" Arguments="-precompile -EnableDSYM $(ExtraCompileArgsMac) $(AllModulesArg) $(MacArchitectureArg)"/>
 			<Compile Target="UnrealEditor" Platform="Mac" Configuration="Development" Tag="#UnrealEditor Mac" Arguments="-precompile -EnableDSYM $(ExtraCompileArgsMac) $(AllModulesArg) $(MacArchitectureArg)"/>
 
 			<!-- Note: We compile LiveLinkHub here because we share a modular build environment with the editor. This drastically reduces its build time and prevents having duplicate files in the MakeInstalledBuild step. -->
-			<!--<Compile Target="LiveLinkHub" Configuration="Development" Platform="Mac" Tag="#UnrealEditor Mac" Arguments="-precompile -EnableDSYM $(MacArchitectureArg)" /> -->
+			<!--<Compile Target="LiveLinkHub" Configuration="Development" Platform="Mac" Tag="#UnrealEditor Mac" Arguments="-precompile -EnableDSYM $(MacArchitectureArg)" />-->
+			<!--<Tag Files="Engine/Source/Programs/LiveLinkHub/Binaries/Mac/LiveLinkHub.target" With="#UnrealEditor Mac"/>-->
 
 			<Property Name="UnrealEditorExceptions" Value=""/>
 			<Property Name="UnrealEditorExceptions" Value="$(PluginsExceptions)" If="$(CompileDatasmithPlugins)"/>
@@ -276,11 +280,13 @@
 	<!-- Base editor and tools for compile on Linux -->
 	<Agent Name="Editor Linux" Type="Linux_Licensee">
 		<Node Name="Compile UnrealEditor Linux" Requires="Update Version Files" Produces="#UnrealEditor Linux;#UnrealEditor Linux Unstripped;#UnrealEditor Linux Stripped;#UnrealEditor Linux WithExceptions">
+			<QueryTargets ProjectFile="Engine/Source/Programs/LiveLinkHub/LiveLinkHub.uproject" Tag="#UnrealEditor Linux" />
 			<Compile Target="UnrealEditor" Platform="Linux" Configuration="DebugGame" Tag="#UnrealEditor Linux" Arguments="-precompile -allmodules $(ExtraCompileArgs)"/>
 			<Compile Target="UnrealEditor" Platform="Linux" Configuration="Development" Tag="#UnrealEditor Linux" Arguments="-precompile -allmodules $(ExtraCompileArgs)"/>
 
 			<!-- Note: We compile LiveLinkHub here because we share a modular build environment with the editor. This drastically reduces its build time and prevents having duplicate files in the MakeInstalledBuild step. -->
-			<!--<Compile Target="LiveLinkHub" Configuration="Development" Platform="Linux" Tag="#UnrealEditor Linux" Arguments="-precompile $(ExtraCompileArgs)" />-->
+			<Compile Target="LiveLinkHub" Configuration="Development" Platform="Linux" Tag="#UnrealEditor Linux" Arguments="-precompile $(ExtraCompileArgs)" />
+			<Tag Files="Engine/Source/Programs/LiveLinkHub/Binaries/Linux/LiveLinkHub.target" With="#UnrealEditor Linux"/>
 
 			<Tag Files="#UnrealEditor Linux" Filter="$(LinuxStripFilter)" With="#UnrealEditor Linux Unstripped"/>
 			<Tag Files="#UnrealEditor Linux" Except="$(PluginsExceptions)" With="#UnrealEditor Linux WithExceptions"/>
