Adding motion design template for 5.5 release.

diff --git a/Engine/Build/InstalledEngineFilters.xml b/Engine/Build/InstalledEngineFilters.xml
--- a/Engine/Build/InstalledEngineFilters.xml
+++ b/Engine/Build/InstalledEngineFilters.xml
@@ -65,6 +65,7 @@
 		Templates/TP_VirtualRealityBP/TP_VirtualRealityBP.uproject
 		Templates/TP_SIM_BlankBP/TP_SIM_BlankBP.uproject
 		Templates/TP_DMXBP/TP_DMXBP.uproject
+		Templates/TP_MotionDesignBP/TP_MotionDesignBP.uproject
 		Samples/StarterContent/StarterContent.uproject
 	</Property>
 	<!-- Win64 specific DDC -->
@@ -136,7 +137,8 @@
 		Templates/TP_SIM_BlankBP/...
 
 		Templates/TP_DMXBP/...
-
+		Templates/TP_MotionDesignBP/...
+		
 		<!-- Shared template resources -->
 		Templates/TemplateResources/...
 
