Move the EditorPerf AssetRegistry Cache location one directory down so we don't run into linux permission issues trying to write in privileged directories

#rnx

diff --git a/Engine/Build/EditorPerf.xml b/Engine/Build/EditorPerf.xml
--- a/Engine/Build/EditorPerf.xml
+++ b/Engine/Build/EditorPerf.xml
@@ -153,7 +153,7 @@
 	<Include Script="../Restricted/NotForLicensees/Build/EditorPerfInternalArgs.xml" If="!$(IsLicensee)" />
 
 	<!-- Option to retain per-branch asset registry between runs so that warm is representative of user experience -->
-	<Property Name="AssetRegistryCache_Folder" Value="$(RootDir)/../../../EditorPerfAssetRegistryCache/$(EscapedBranch)/"/>
+	<Property Name="AssetRegistryCache_Folder" Value="$(RootDir)/../../EditorPerfAssetRegistryCache/$(EscapedBranch)/"/>
 
 	<!-- Options for extra runtime args -->
 	<Option Name="AdditionalRuntimeArgs" DefaultValue="" Description="Additional Arguments to pass to RunTest"/>
