Fixing dotnet6 being included in Installed Engine Builds

#rnx

diff --git a/Engine/Build/InstalledEngineFilters.xml b/Engine/Build/InstalledEngineFilters.xml
--- a/Engine/Build/InstalledEngineFilters.xml
+++ b/Engine/Build/InstalledEngineFilters.xml
@@ -663,8 +663,8 @@
 
 	<!-- TVOS -->
 	<Property Name="ExtraSignFilesTVOS">
-		Engine/Binaries/ThirdParty/DotNet/6.0.302/mac-arm64/...
-		Engine/Binaries/ThirdParty/DotNet/6.0.302/mac-x64/...
+		Engine/Binaries/ThirdParty/DotNet/8.0.300/mac-arm64/...
+		Engine/Binaries/ThirdParty/DotNet/8.0.300/mac-x64/...
 		Engine/Extras/ThirdPartyNotUE/libimobiledevice/x64/...
 	</Property>
 
