Installed Build: Fixing Invalid VisionOS References

#rnx

diff --git a/Engine/Build/InstalledEngineBuild.xml b/Engine/Build/InstalledEngineBuild.xml
--- a/Engine/Build/InstalledEngineBuild.xml
+++ b/Engine/Build/InstalledEngineBuild.xml
@@ -153,8 +153,8 @@
 	<Macro Name="Platform_Append_DDCPlatformsWin64"/>
 	<Macro Name="Platform_Append_DDCPlatformsMac"/>
 	<Macro Name="Platform_TagDDCFiles"/>
-	<Macro Name="Platform_Tag_UnrealGame"/>
-	<Macro Name="Platform_ExtraSignFiles"/>
+	<Macro Name="Platform_Tag_UnrealGame" Arguments="HostPlatform"/>
+	<Macro Name="Platform_ExtraSignFiles" Arguments="HostPlatform"/>
 
 	<!-- Include any scripts from platform extensions -->
 	<Property Name="ExtensionPlatforms" Value=""/>
@@ -1083,8 +1083,8 @@
 				<Property Name="CopyInstalledExceptions" Value="$(CopyInstalledExceptions);$(CopyLinuxExceptions)"/>
 				<Property Name="ExtraSignFiles" Value="$(ExtraSignFiles);$(ExtraSignFilesLinux)"/>
 			</Do>
-			<Expand Name="Platform_Tag_UnrealGame"/>
-			<Expand Name="Platform_ExtraSignFiles"/>
+			<Expand Name="Platform_Tag_UnrealGame" HostPlatform="Win64"/>
+			<Expand Name="Platform_ExtraSignFiles" HostPlatform="Win64"/>
 
 			<!-- Tag any dependencies from all previous build products -->
 			<Tag Files="#Installed Win64" Filter="*.target" With="#TargetReceipts"/>
@@ -1225,8 +1225,8 @@
 				<Property Name="CopyInstalledFilter" Value="$(CopyInstalledFilter);$(CopyTVOSFilter)"/>
 				<Property Name="CopyInstalledExceptions" Value="$(CopyInstalledExceptions);$(CopyTVOSExceptions)"/>
 			</Do>
-			<Expand Name="Platform_Tag_UnrealGame" />
-			<Expand Name="Platform_ExtraSignFiles" />
+			<Expand Name="Platform_Tag_UnrealGame" HostPlatform="Mac" />
+			<Expand Name="Platform_ExtraSignFiles" HostPlatform="Mac" />
 
 			<!-- Tag any dependencies from all previous build products -->
 			<Tag Files="#Installed Mac" Filter="*.target" With="#TargetReceipts"/>
@@ -1347,8 +1347,8 @@
 				<Property Name="CopyInstalledFilter" Value="$(CopyInstalledFilter);$(CopyAndroidFilter)"/>
 				<Property Name="CopyInstalledExceptions" Value="$(CopyInstalledExceptions);$(CopyAndroidExceptions)"/>
 			</Do>
-			<Expand Name="Platform_Tag_UnrealGame" />
-			<Expand Name="Platform_ExtraSignFiles"/>
+			<Expand Name="Platform_Tag_UnrealGame" HostPlatform="Linux"/>
+			<Expand Name="Platform_ExtraSignFiles" HostPlatform="Linux"/>
 
 			<!-- Tag any dependencies from all previous build products -->
 			<Tag Files="#Installed Linux" Filter="*.target" With="#TargetReceipts"/>
diff --git a/Engine/Platforms/VisionOS/Build/InstalledEngineBuild.xml b/Engine/Platforms/VisionOS/Build/InstalledEngineBuild.xml
--- a/Engine/Platforms/VisionOS/Build/InstalledEngineBuild.xml
+++ b/Engine/Platforms/VisionOS/Build/InstalledEngineBuild.xml
@@ -3,7 +3,8 @@
 
 	<Option Name="WithVisionOS" Restrict="true|false" DefaultValue="false" Description="Include the VisionOS target platform"/>
 
-	<Property Name="ExtensionPlatforms" Value="$(ExtensionPlatforms);" If="'$(ExtensionPlatforms)' != ''"/>	<Property Name="ExtensionPlatforms" Value="$(ExtensionPlatforms)VisionOS"/>
+	<Property Name="ExtensionPlatforms" Value="$(ExtensionPlatforms);" If="'$(ExtensionPlatforms)' != ''"/>
+	<Property Name="ExtensionPlatforms" Value="$(ExtensionPlatforms)VisionOS"/>
 
 	<Do If="$(WithVisionOS)">
 		<Expand Name="AddExtensionPlatform" PlatformDir="VisionOS"/>
@@ -67,7 +68,7 @@
 	</Extend>
 
 	<Extend Name="Platform_Tag_UnrealGame">
-		<Do If="$(WithVisionOS)">
+		<Do If="$(WithVisionOS) And '$(HostPlatform)' == 'Mac'">
 			<Property Name="CopyVisionOSFilter">
 				<!-- Build Files -->
 				Engine/Build/VisionOS/...
