[EditorPerfHordeJob] No longer copy Trace files to a remote share and instead rely on Horde Telemetry storage

#rnx

diff --git a/Engine/Build/EditorPerf.xml b/Engine/Build/EditorPerf.xml
--- a/Engine/Build/EditorPerf.xml
+++ b/Engine/Build/EditorPerf.xml
@@ -14,7 +14,6 @@
 	<Property Name="BuildName" Value="$(EscapedBranch)-CL-$(Change)$(PreflightSuffix)"/>
 	<Option Name="CompileCategory" DefaultValue="Editors" Description="Category string for Compile Label"/>
 	<Option Name="RunCategory" DefaultValue="Test" Description="Category string for Run Label"/>
-	<Option Name="BaselineCategory" DefaultValue="Test" Description="Category string for Baseline Label"/>
 	<Option Name="SingleAgentCompile" Restrict="true|false" DefaultValue="false" Description="Run all compile steps on same agent"/>
 	<Option Name="AdditionalEditorCompileArgs" DefaultValue="" Description="Additional Compile Arguments for the Editor"/>
 
@@ -53,18 +52,9 @@
 	<Property Name="AutomationNetworkRootDir" Value="$(NetworkReportRoot)/EditorPerf" If="'$(HostPlatform)' != 'Win64'"/>
 	<Property Name="AutomationLocalOutputDir" Value="$(RootDir)/LocalBuilds/Reports"/>
 
-	<!-- Outputs for this run keeping slashes flipped the right way because it shows up on the webapp this way -->
-	<Property Name="AutomationNetworkOutputDir" Value="$(AutomationNetworkRootDir)\$(BuildName)"/>
-	<Property Name="AutomationNetworkOutputDir" Value="$(AutomationNetworkRootDir)/$(BuildName)" If="'$(HostPlatform)' != 'Win64'"/>
+	<!-- Outputs for this run -->
 	<Property Name="AutomationOutputDir" Value="$(AutomationLocalOutputDir)"/>
-	<Property Name="AutomationOutputDir" Value="$(AutomationNetworkOutputDir)" If="$(IsBuildMachine)"/>
-
-	<!-- Options for baselining -->
-	<Option Name="RunAndBaseline" DefaultValue="false" Description="Enable graph dependencies for running and setting the baseline in a single build."/>
-	<Option Name="BaselineDir" DefaultValue="$(RootDir)/Engine/Build/Baseline" Description="The revision controlled location of the baseline telemetry files."/>
-	<Option Name="BaselineBuildName" DefaultValue="$(BuildName)" Description="The BuildName of the build to set as the new baseline to compare against."/>
-	<Property Name="BaselineAutomationNetworkOutputDir" Value="$(AutomationNetworkRootDir)/$(BaselineBuildName)"/>
-	<Property Name="BaselineAutomationNetworkOutputDir" Value="$(AutomationOutputDir)" If="$(RunAndBaseline)"/>
+	<Property Name="AutomationOutputDir" Value="$(RootDir)/Engine/Programs/AutomationTool/Saved/Logs" If="$(IsBuildMachine)"/>
 
 	<!-- Options for warm+cold vs cold cache only -->
 	<Option Name="IncludeWarmPhase" DefaultValue="true" Description="Whether to include a warm cache phase. Default is true and will run cold and warm tests. If false, only a cold test will run."/>
@@ -72,7 +62,7 @@
 	<Property Name="Phases" Value="Cold" If="'$(IncludeWarmPhase)'=='false'"/>
 
 	<!-- Arguments for tracing -->
-	<Property Name="TraceDir" Value="$(RootDir)/Engine/Programs/UnrealInsights/Saved/TraceSessions"/>
+	<Property Name="TraceDir" Value="$(RootDir)/Engine/Programs/AutomationTool/Saved/Logs/Traces"/> <!-- Should write to 'Telemetry' instead of 'Logs/Traces' but we cannot until UE-228855 -->
 
 	<!-- Arguments for phases, keep in mind that these are passed to both the DDCCleanupCommandlet and RunUnreal -->
 	<Option Name="EnableBasePhaseTests" DefaultValue="true" Description="Enable the base phase perf tests."/>
@@ -209,7 +199,6 @@
 	<!-- Gauntlet arguments for commandlet, summarize runs -->
 	<Property Name="CommandletArgs" Value="-notraceserver"/>
 	<Property Name="SummarizeArgs" Value=""/>
-	<Property Name="SummarizeArgs" Value="-skipbaseline" If="$(RunAndBaseline)"/>
 
 	<!-- Publish telemetry options -->
 	<Option Name="TelemetryConfig" DefaultValue="" Description="The configuration to use when publishing telemetry data to a database"/>
@@ -286,28 +275,9 @@
 		<Property Name="GenerateTSANReportNodeList" Value="$(GenerateTSANReportNodeList); $(GenerateTSANReportNodeName)"/>
 	</Macro>
 
-	<Macro Name="UploadTestTraces" Arguments="Platform;ProjectName">
-		<!-- Save all files named for this project in the trace directory -->
-		<Property Name="TestOutputDir" Value="$(AutomationOutputDir)/$(ProjectName)$(Platform)EditorPerf"/>
-		<Copy From="$(TraceDir)/$(ProjectName)$(Platform)EditorPerf*.*" To="$(TestOutputDir)"/>
-	</Macro>
-
 	<!-- These will be set to the names of all nodes we want to execute based on specified platforms/options -->
 	<Property Name="RunGraphNodes" Value=""/>
 
-	<!-- Mimic of the Runner that retrieves the stashed telemetry file and copies it to the depot location -->
-	<Macro Name="BaselineResult" Arguments="Platform;Project;Phase">
-		<Property Name="TestName" Value="EditorPerf$(Phase)"/>
-		<Property Name="IterationName" Value="$(Project)$(Platform)$(TestName)1"/> <!-- Note the 1! We use Iteration 1 for baselines for now -->
-		<Property Name="TelemetryFileName" Value="$(IterationName)Telemetry.csv"/>
-		<Property Name="LocalTelemetryFilePath" Value="$(BaselineDir)/$(TelemetryFileName)"/>
-		<Property Name="NetworkTelemetryFilePath" Value="$(BaselineAutomationNetworkOutputDir)/$(TelemetryFileName)"/>
-		<Copy From="$(NetworkTelemetryFilePath)" To="$(LocalTelemetryFilePath)" Overwrite="true"/>
-		<Do If="!$(IsBuildMachine)">
-			<Spawn Exe="p4" Arguments="reconcile -ea $(LocalTelemetryFilePath)"/>
-		</Do>
-	</Macro>
-
 	<Property Name="NeedToCopyTSanBinaries" Value="$(WithTSAN)"/>
 	<Macro Name="CopyTSANBinariesIfNeeded" Arguments="ProjectBinariesPath">
 		<!-- Perform swap of built artifacts for TSAN builds so we use a non-TSAN build for cold runs and a TSAN build for warm runs -->
@@ -315,8 +285,6 @@
 		<Property Name="NeedToCopyTSanBinaries" Value="false"/>
 	</Macro>
 	
-	<!-- This will be set to the names of all nodes we want to execute based on specified platforms/options -->
-	<Property Name="BaselineGraphNodes" Value=""/>
 
 	<!-- ForEach Platform to run, generate graph nodes -->
 	<ForEach Name="Platform" Values="$(Platforms)">
@@ -488,31 +456,9 @@
 							<Property Name="RunNodeList" Value="$(RunNodeList);$(RunNodeLazyLoadZenName)" If="$(EnableLazyLoadZenLoaderTest)"/>
 						</ForEach>
 					</Do>
-
-					<Do If="!$(WithTSAN)">
-						<Property Name="RunNodeUploadPhaseTraces" Value="Upload $(OutputName) Traces"/>
-						<Node Name="$(RunNodeUploadPhaseTraces)">
-							<Expand Name="UploadTestTraces" Platform="$(Platform)" ProjectName="$(Project)"/>
-						</Node>
-						<Property Name="UploadNodeList" Value="$(RunNodeUploadPhaseTraces)"/>
-					</Do>
 				</Agent>
 				<Label Category="$(RunCategory)" Name="$(Project) Editor Run" Requires="$(RunNodeList);$(UploadNodeList)" UgsBadge="$(Project) Editor Run"/>
 				<Property Name="RunGraphNodes" Value="$(RunGraphNodes);$(RunNodeList);$(GenerateTSANReportNodeList);$(UploadNodeList)"/>
-
-				<!--Baseline the Results--> 
-				<Property Name="BaselineNodeName" Value="Baseline $(OutputName)"/>
-				<Property Name="BaselineNodeRequires" Value=""/>
-				<Property Name="BaselineNodeRequires" Value="$(RunNodeName)" If="$(RunAndBaseline)"/>
-				<Property Name="BaselineGraphNodes" Value="$(BaselineGraphNodes);$(BaselineNodeName)"/>
-				<Agent Name="EditorPerf Agent $(Platform) $(Project)" Type="EditorPerf$(Platform)">
-					<Node Name="$(BaselineNodeName)" Requires="$(BaselineNodeRequires)">
-						<ForEach Name="Phase" Values="$(Phases)">
-							<Expand Name="BaselineResult" Platform="$(Platform)" Project="$(Project)" Phase="$(Phase)"/>
-						</ForEach>
-					</Node>
-				</Agent>
-				<Label Category="$(BaselineCategory)" Name="$(Project) Perf Baseline" Requires="$(BaselineNodeName)"/>
 			</Do>
 		</ForEach> <!-- Projects -->
 
@@ -520,17 +466,4 @@
 
 	<!-- These Nodes run the Editor and gather profile data -->
 	<Aggregate Name="Run EditorPerf Tests" Requires="$(RunGraphNodes)"/>
-
-	<!-- These Nodes copies the CSV telemetry files to special location in the depot, thus making them the "baselines" -->
-	<Aggregate Name="Baseline EditorPerf Results" Requires="$(BaselineGraphNodes)"/>
-
-	<!-- This Node submits the baseline CSV files to Perforce -->
-	<!--
-	<Agent Name="Submission Agent Win64" Type="Win64">
-		<Node Name="Submit Baseline To Perforce" Requires="$(BaselineGraphNodes)">
-			<Submit Description="[CL $(Change)] Updated EditorPerf baselines" Files="$(BaselineFiles)"/>
-		</Node>
-	</Agent>
-	-->
-
 </BuildGraph>
