Store TSAN logs in a network share so the log parser may find them as we cannot path to S3 logs created by another node action

#rnx

diff --git a/Engine/Build/EditorPerf.xml b/Engine/Build/EditorPerf.xml
--- a/Engine/Build/EditorPerf.xml
+++ b/Engine/Build/EditorPerf.xml
@@ -55,6 +55,11 @@
 	<!-- Outputs for this run -->
 	<Property Name="AutomationOutputDir" Value="$(AutomationLocalOutputDir)"/>
 	<Property Name="AutomationOutputDir" Value="$(RootDir)/Engine/Programs/AutomationTool/Saved/Logs" If="$(IsBuildMachine)"/>
+	<!-- We need to share logs between nodes for TSAN runs so store the logs to a share rather than Horde's Log storage -->
+	<Property Name="AutomationNetworkOutputDir" Value="$(AutomationNetworkRootDir)\$(BuildName)" If="$(WithTSAN)"/>
+	<Property Name="AutomationNetworkOutputDir" Value="$(AutomationNetworkRootDir)/$(BuildName)" If="$(WithTSAN) and '$(HostPlatform)' != 'Win64'"/>
+	<Property Name="AutomationOutputDir" Value="$(AutomationLocalOutputDir)" If="$(WithTSAN)"/>
+	<Property Name="AutomationOutputDir" Value="$(AutomationNetworkOutputDir)" If="$(WithTSAN) and $(IsBuildMachine)"/>
 
 	<!-- Options for warm+cold vs cold cache only -->
 	<Option Name="IncludeWarmPhase" DefaultValue="true" Description="Whether to include a warm cache phase. Default is true and will run cold and warm tests. If false, only a cold test will run."/>
@@ -62,7 +67,7 @@
 	<Property Name="Phases" Value="Cold" If="'$(IncludeWarmPhase)'=='false'"/>
 
 	<!-- Arguments for tracing -->
-	<Property Name="TraceDir" Value="$(RootDir)/Engine/Programs/AutomationTool/Saved/Logs/Traces"/> <!-- Should write to 'Telemetry' instead of 'Logs/Traces' but we cannot until UE-228855 -->
+	<Property Name="TraceDir" Value="$(RootDir)/Engine/Programs/AutomationTool/Saved/Logs/Traces"/>
 
 	<!-- Arguments for phases, keep in mind that these are passed to both the DDCCleanupCommandlet and RunUnreal -->
 	<Option Name="EnableBasePhaseTests" DefaultValue="true" Description="Enable the base phase perf tests."/>
