Added OnlineTests support on the Mac OS.

diff --git a/Engine/Build/LowLevelTests.xml b/Engine/Build/LowLevelTests.xml
--- a/Engine/Build/LowLevelTests.xml
+++ b/Engine/Build/LowLevelTests.xml
@@ -145,7 +145,9 @@
 				<Property Name="ExtraCompilationArgsFinal" Value="$(ExtraCompilationArgs) -projectprogramoverride=$(RootDir)/$(BuildProject)"/>
 			</Do>
 			<Property Name="ExtraCompilationArgsFinal" Value="$(ExtraCompilationArgsFinal) -ubtargs=&quot;-UBA &quot;" If="$(WithUBA)"/>
-			<Command Name="BuildCookRun" Arguments="-Build -NoBootstrapExe -NoCodeSign -SkipCook -SkipBuildEditor=True -Stage -NoSubmit -Archive -ArchiveDirectory=$(LocalAchiveDirectory) -project=$(TargetName) -configuration=$(Configuration) -platform=$(Platform) $(ExtraCompilationArgsFinal)" If="!$(DryRun)"/>
+            <Property Name="BuildCookRunArgs" Value="-Build -NoBootstrapExe -NoCodeSign -SkipCook -SkipBuildEditor=True -Stage -NoSubmit -Archive -ArchiveDirectory=$(LocalAchiveDirectory) -project=$(TargetName) -configuration=$(Configuration) -platform=$(Platform) $(ExtraCompilationArgsFinal)"/>
+            <Property Name="BuildCookRunArgs" Value="$(BuildCookRunArgs) -package" If="'$(BuildPlatform)' == 'Mac'"/>
+			<Command Name="BuildCookRun" Arguments="$(BuildCookRunArgs)" If="!$(DryRun)"/>
 		</Do>
 		<Do If="'$(StagesWithProjectFile)' != 'True'">
 		<Property Name="ExtraCompilationArgsFinal" Value="$(ExtraCompilationArgsFinal) -UBA" If="$(WithUBA)"/>
