Stop building .net4 projects

#rnx

diff --git a/Engine/Build/InstalledEngineBuild.xml b/Engine/Build/InstalledEngineBuild.xml
--- a/Engine/Build/InstalledEngineBuild.xml
+++ b/Engine/Build/InstalledEngineBuild.xml
@@ -572,6 +572,7 @@
 
 			<!-- Compile all the tools -->
 			<CsCompile Project="#RedistUATProjects" Configuration="Development" Platform="AnyCPU" Tag="#Build Tools CS Binaries;#Sign Binaries" TagReferences="#Build Tools CS Binaries"/>
+			<!-- TODO: needs to be updated to dotnet8 or deleted
 			<Do If="'$(HostPlatform)' == 'Win64'">
 				<CsCompile Project="$(CsToolsDir)/Engine/Source/Programs/UnrealSwarm/SwarmCoordinator/SwarmCoordinator.csproj" Configuration="Development" Platform="AnyCPU" UseSystemCompiler="true" Tag="#Build Tools CS Binaries;#Sign Binaries" TagReferences="#Build Tools CS Binaries"/>
 				<CsCompile Project="$(CsToolsDir)/Engine/Source/Programs/UnrealSwarm/Agent/Agent.csproj" Configuration="Development" Platform="AnyCPU" UseSystemCompiler="true" Tag="#Build Tools CS Binaries;#Sign Binaries" TagReferences="#Build Tools CS Binaries"/>
@@ -581,6 +582,7 @@
 				<CsCompile Project="$(CsToolsDir)/Engine/Source/Programs/NetworkProfiler/NetworkProfiler/NetworkProfiler.csproj" Configuration="Development" Platform="AnyCPU" UseSystemCompiler="true" Tag="#Build Tools CS Binaries" TagReferences="#Build Tools CS Binaries;#Sign Binaries"/>
 				<Expand Name="Platform_Build_Tools_CS_Binaries"/>
 			</Do>
+			-->
 
 			<!-- Save and tag BuildConfiguration.Schema.xsd as this can/will be overriden by other processes piggy backing on cstools to get it to the right spot -->
 			<Copy Files="$(RootDir)/Engine/Saved/UnrealBuildTool/BuildConfiguration.Schema.xsd" From="$(RootDir)" To="$(CsToolsDir)"/>
