Add option to target compile agents when building nodes with LLT.

#rnx

diff --git a/Engine/Build/LowLevelTests.xml b/Engine/Build/LowLevelTests.xml
--- a/Engine/Build/LowLevelTests.xml
+++ b/Engine/Build/LowLevelTests.xml
@@ -12,8 +12,9 @@
 	<Property Name="PlatformsBuiltFromLinux" Value="Linux;LinuxArm64"/>
 
 	<Property Name="AlternateDevicePoolPlatformsProperty" Value=""/>
-
 	<Property Name="AdbKeysPath" Value=""/>
+	
+	<Option Name="UseCompileAgents" DefaultValue="false" Description="Targets compile agents instead of regular agents for Build Platforms"/>
 
 	<Macro Name="SelectAdbKeysPath" Arguments="BuildPlatform">
 		<!-- Extend this macro and set AdbKeysPath -->
@@ -31,10 +32,12 @@
 		</Do>
 		<Do If="'$(Platform)' == 'Mac' Or '$(Platform)' == 'iOS'">
 			<Property Name="AgentType" Value="Mac"/>
+			<Property Name="AgentType" Value="MacCompile" If="$(UseCompileAgents)"/>
 			<Property Name="BuildPlatform" Value="Mac"/>
 		</Do>
 		<Do If="'$(Platform)' == 'Win64' Or '$(Platform)' == 'Android'">
 			<Property Name="AgentType" Value="Win64"/>
+			<Property Name="AgentType" Value="Win64Compile" If="$(UseCompileAgents)"/>
 			<Property Name="BuildPlatform" Value="Win64"/>
 		</Do>
 	</Macro>
@@ -58,7 +61,7 @@
 	</Macro>
 
 	<Macro Name="RunAllTests"/>
-	
+		
 	<Option Name="BuildCatch2Libs" DefaultValue="false" Description="True if building Catch2 static libraries"/>
 	<Option Name="TestCatch2Libs" DefaultValue="false" Description="True if additionally building Catch2 with extras and tests"/>
 	<Option Name="Catch2LibVariation" DefaultValue="VS2019" Description="Catch2 library variation, e.g. VS2022"/>
@@ -243,7 +246,7 @@
 					</Do>
 				</Do>
 			</Node>
-			<Label Category="Test" Name="$(ShortName) $(Platform)" Requires="$(ShortName) Tests $(Platform)"/>
+			<Label Category="Tests" Name="$(ShortName) $(Platform)" Requires="$(ShortName) Tests $(Platform)"/>
 			<Trace Message="$(ShortName) runs because $(TestRuns)" If="$(TestRuns) and $(DryRun)"/>
 			<Do If="$(TestRuns)">
 				<Property Name="TestNodes" Value="$(TestNodes);$(ShortName) Tests $(Platform)"/>
