Exclude Automation Tool project obj folders with deeper nesting

Prevents uneeded compile intermediate files from being included in an installed engine build

#rnx

diff --git a/Engine/Build/InstalledEngineFilters.xml b/Engine/Build/InstalledEngineFilters.xml
--- a/Engine/Build/InstalledEngineFilters.xml
+++ b/Engine/Build/InstalledEngineFilters.xml
@@ -532,10 +532,10 @@
 	</Property>
 	<Property Name="CopyWin64Exceptions">
 		Engine/Platforms/*/Source/Programs/AutomationTool/obj/...
-		Engine/Platforms/*/Source/Programs/AutomationTool/*/obj/...	
+		Engine/Platforms/*/Source/Programs/AutomationTool/.../obj/...
 		Engine/Platforms/*/Source/Programs/UnrealBuildTool/obj/...
 		Engine/Source/Programs/AutomationTool/obj/...
-		Engine/Source/Programs/AutomationTool/*/obj/...
+		Engine/Source/Programs/AutomationTool/.../obj/...
 		Engine/Source/Programs/UnrealBuildTool/obj/...
 		Engine/Source/Programs/Shared/*/bin/...
 		Engine/Source/Programs/Shared/*/obj/...
@@ -555,10 +555,10 @@
 	</Property>
 	<Property Name="CopyMacExceptions">
 		Engine/Platforms/*/Source/Programs/AutomationTool/obj/...
-		Engine/Platforms/*/Source/Programs/AutomationTool/*/obj/...
+		Engine/Platforms/*/Source/Programs/AutomationTool/.../obj/...
 		Engine/Platforms/*/Source/Programs/UnrealBuildTool/obj/...
 		Engine/Source/Programs/AutomationTool/obj/...
-		Engine/Source/Programs/AutomationTool/*/obj/...
+		Engine/Source/Programs/AutomationTool/.../obj/...
 		Engine/Source/Programs/UnrealBuildTool/obj/...
 		Engine/Source/Programs/Shared/*/bin/...
 		Engine/Source/Programs/Shared/*/obj/...
@@ -716,10 +716,10 @@
 
 	<Property Name="CopyLinuxExceptions">
 		Engine/Platforms/*/Source/Programs/AutomationTool/obj/...
-		Engine/Platforms/*/Source/Programs/AutomationTool/*/obj/...
+		Engine/Platforms/*/Source/Programs/AutomationTool/.../obj/...
 		Engine/Platforms/*/Source/Programs/UnrealBuildTool/obj/...
 		Engine/Source/Programs/AutomationTool/obj/...
-		Engine/Source/Programs/AutomationTool/*/obj/...
+		Engine/Source/Programs/AutomationTool/.../obj/...
 		Engine/Source/Programs/UnrealBuildTool/obj/...
 		Engine/Source/Programs/Shared/*/bin/...
 		Engine/Source/Programs/Shared/*/obj/...
