Exclude AutomationTool and UnrealBuildTool bin directories under Engine\Source from installed build

#rnx

diff --git a/Engine/Build/InstalledEngineFilters.xml b/Engine/Build/InstalledEngineFilters.xml
--- a/Engine/Build/InstalledEngineFilters.xml
+++ b/Engine/Build/InstalledEngineFilters.xml
@@ -529,11 +529,17 @@
 		Engine/Source/ThirdParty/EOSSDK/SDK/Bin/x64/xaudio2_9redist.dll
 	</Property>
 	<Property Name="CopyWin64Exceptions">
+		Engine/Platforms/*/Source/Programs/AutomationTool/bin/...
 		Engine/Platforms/*/Source/Programs/AutomationTool/obj/...
+		Engine/Platforms/*/Source/Programs/AutomationTool/.../bin/...
 		Engine/Platforms/*/Source/Programs/AutomationTool/.../obj/...
+		Engine/Platforms/*/Source/Programs/UnrealBuildTool/bin/...
 		Engine/Platforms/*/Source/Programs/UnrealBuildTool/obj/...
+		Engine/Source/Programs/AutomationTool/bin/...
 		Engine/Source/Programs/AutomationTool/obj/...
+		Engine/Source/Programs/AutomationTool/.../bin/...
 		Engine/Source/Programs/AutomationTool/.../obj/...
+		Engine/Source/Programs/UnrealBuildTool/bin/...
 		Engine/Source/Programs/UnrealBuildTool/obj/...
 		Engine/Source/Programs/Shared/*/bin/...
 		Engine/Source/Programs/Shared/*/obj/...
@@ -552,11 +558,17 @@
 		Engine/Source/ThirdParty/PLCrashReporter/lib/lib-Xcode-12.4/Mac/...
 	</Property>
 	<Property Name="CopyMacExceptions">
+		Engine/Platforms/*/Source/Programs/AutomationTool/bin/...
 		Engine/Platforms/*/Source/Programs/AutomationTool/obj/...
+		Engine/Platforms/*/Source/Programs/AutomationTool/.../bin/...
 		Engine/Platforms/*/Source/Programs/AutomationTool/.../obj/...
+		Engine/Platforms/*/Source/Programs/UnrealBuildTool/bin/...
 		Engine/Platforms/*/Source/Programs/UnrealBuildTool/obj/...
+		Engine/Source/Programs/AutomationTool/bin/...
 		Engine/Source/Programs/AutomationTool/obj/...
+		Engine/Source/Programs/AutomationTool/.../bin/...
 		Engine/Source/Programs/AutomationTool/.../obj/...
+		Engine/Source/Programs/UnrealBuildTool/bin/...
 		Engine/Source/Programs/UnrealBuildTool/obj/...
 		Engine/Source/Programs/Shared/*/bin/...
 		Engine/Source/Programs/Shared/*/obj/...
@@ -710,11 +722,17 @@
 	</Property>
 
 	<Property Name="CopyLinuxExceptions">
+		Engine/Platforms/*/Source/Programs/AutomationTool/bin/...
 		Engine/Platforms/*/Source/Programs/AutomationTool/obj/...
+		Engine/Platforms/*/Source/Programs/AutomationTool/.../bin/...
 		Engine/Platforms/*/Source/Programs/AutomationTool/.../obj/...
+		Engine/Platforms/*/Source/Programs/UnrealBuildTool/bin/...
 		Engine/Platforms/*/Source/Programs/UnrealBuildTool/obj/...
+		Engine/Source/Programs/AutomationTool/bin/...
 		Engine/Source/Programs/AutomationTool/obj/...
+		Engine/Source/Programs/AutomationTool/.../bin/...
 		Engine/Source/Programs/AutomationTool/.../obj/...
+		Engine/Source/Programs/UnrealBuildTool/bin/...
 		Engine/Source/Programs/UnrealBuildTool/obj/...
 		Engine/Source/Programs/Shared/*/bin/...
 		Engine/Source/Programs/Shared/*/obj/...
