Ensure that UnrealPackageTool is compiled in regular low-cadence CI jobs

diff --git a/Engine/Build/InstalledEngineBuild.xml b/Engine/Build/InstalledEngineBuild.xml
--- a/Engine/Build/InstalledEngineBuild.xml
+++ b/Engine/Build/InstalledEngineBuild.xml
@@ -540,6 +540,7 @@
 			<Compile Target="SwitchboardListener" Configuration="Development" Platform="Win64" Tag="#Build Tools Win64" Arguments="$(ExtraCompileArgs)"/>
 			<Compile Target="SwitchboardListenerHelper" Configuration="Development" Platform="Win64" Tag="#Build Tools Win64" Arguments="$(ExtraCompileArgs)"/>
 			<Compile Target="UnrealObjectPtrTool" Configuration="Development" Platform="Win64" Tag="#Build Tools Win64" Arguments="$(ExtraCompileArgs)"/>
+			<Compile Target="UnrealPackageTool" Configuration="Development" Platform="Win64" Tag="#Build Tools Win64" Arguments="$(ExtraCompileArgs)"/>
 			<Compile Target="ZenDashboard" Configuration="Development" Platform="Win64" Tag="#Build Tools Win64" Arguments="$(ExtraCompileArgs)"/>
 
 			<Tag Files="Engine/Binaries/Win64/XGEControlWorker.exe" With="#Build Tools Win64"/>
