Added option to include source code docs in the installed engine build

#rnx

diff --git a/Engine/Build/InstalledEngineBuild.xml b/Engine/Build/InstalledEngineBuild.xml
--- a/Engine/Build/InstalledEngineBuild.xml
+++ b/Engine/Build/InstalledEngineBuild.xml
@@ -187,6 +187,9 @@
 
 	<Property Name="MacArchitectureArg" Value="-architecture=arm64+x64" />
 
+	<!-- Option to include source code docs in the resulting installed engine build -->
+	<Option Name="IncludeDocs" Restrict="true|false" DefaultValue="true" Description="Include the source code docs in the resulting installed engine build"/>
+
 	<!-- Base editor and tools for compile on Windows -->
 	<Agent Name="Editor Win64" Type="Win64_Licensee">
 		<Node Name="Update Version Files">
diff --git a/Engine/Build/InstalledEngineFilters.xml b/Engine/Build/InstalledEngineFilters.xml
--- a/Engine/Build/InstalledEngineFilters.xml
+++ b/Engine/Build/InstalledEngineFilters.xml
@@ -211,8 +211,8 @@
 	<Property Name="CopyEditorFilter" Value="$(CopyEditorFilter);Engine/Source/UnrealClient.Target.cs" If="$(WithClient)"/>
 	<Property Name="CopyEditorFilter" Value="$(CopyEditorFilter);Engine/Source/UnrealServer.Target.cs" If="$(WithServer)"/>
 	<!-- Optional API docs -->
-	<Property Name="CopyEditorFilter" Value="$(CopyEditorFilter);Engine/Documentation/Builds/CppAPI-HTML.tgz"/>
-	<Property Name="CopyEditorFilter" Value="$(CopyEditorFilter);Engine/Documentation/Builds/BlueprintAPI-HTML.tgz"/>
+	<Property Name="CopyEditorFilter" Value="$(CopyEditorFilter);Engine/Documentation/Builds/CppAPI-HTML.tgz" If="$(IncludeDocs)"/>
+	<Property Name="CopyEditorFilter" Value="$(CopyEditorFilter);Engine/Documentation/Builds/BlueprintAPI-HTML.tgz" If="$(IncludeDocs)"/>
 	<!-- Define Editor Exceptions -->
 	<Property Name="CopyEditorExceptions">
 	
