Added support in the installed engine build to publish cloud artifacts - set PublishCloudArtifacts=true
Also added ability to override project name used at upload if it doesn't match the uproject file (or the uproject file doesn't exist like in this case of publishing installed engine builds)

diff --git a/Engine/Build/InstalledEngineBuild.xml b/Engine/Build/InstalledEngineBuild.xml
--- a/Engine/Build/InstalledEngineBuild.xml
+++ b/Engine/Build/InstalledEngineBuild.xml
@@ -36,6 +36,12 @@
 	<Option Name="ExtraCompileArgsMac" DefaultValue="" Description="Extra arguments to pass to compiles run on Mac platforms"/>
 	<Option Name="ExtraDDCArgs" DefaultValue="" Description="Extra arguments to pass to DDC Generation"/>
 
+	<Option Name="PublishCloudArtifacts" DefaultValue="false" Description="Enable to publish cloud artifacts" />
+
+	<Option Name="CloudPublishBaseline" DefaultValue="" Description="The baseline branch (your mainline) if incorrectly set your data will not de-dup across branches. Only applies if `PublishCloudArtifacts` is set." />
+	<Option Name="CloudPublishHost" DefaultValue="" Description="The host to publish cloud artifacts to if enabled using `PublishCloudArtifacts`" />
+	<Option Name="CloudPublishNamespace" DefaultValue="" Description="The namespace to publish cloud artifacts to if enabled using `PublishCloudArtifacts`" />
+
 	<Option Name="WithWin64NoPCH" DefaultValue="true" Description="Compile Win64 with NoPCH"/>
 	<Property Name="Win64NoPCHArg" Value="" />
 	<Property Name="Win64NoPCHArg" Value="-NoPCH" If="$(WithWin64NoPCH)"/>
@@ -1162,6 +1168,9 @@
 		<Node Name="Upload To Horde Win64" Requires="Make Installed Build Win64">
 			<CreateArtifact Name="installed-build-win64" Type="installed-build" Description="Installed Build (Windows)" BaseDir="$(LocalInstalledDir)" Files="..."/>
 		</Node>
+		<Node Name="Upload To Cloud Win64" Requires="Make Installed Build Win64">
+			<CreateCloudArtifact Name="installed-engine-win64" Type="installed-engine" Description="Installed Build (Windows)" BaseDir="$(LocalInstalledDir)" Files="..." ProjectNameOverride="UE"  Platform="Windows" Namespace="$(CloudPublishNamespace)" Host="$(CloudPublishHost)" BaselineBranch="$(CloudPublishBaseline)" />
+		</Node>
 	</Agent>
 	<Aggregate Name="HostPlatforms_Win64" Label="Builds/Win64" Requires="Make Installed Build Win64" Exclude="Update Version Files"/>
 
@@ -1293,6 +1302,9 @@
 		<Node Name="Upload To Horde Mac" Requires="Make Installed Build Mac">
 			<CreateArtifact Name="installed-build-mac" Type="installed-build" Description="Installed Build (Mac)" BaseDir="$(LocalInstalledDirMac)" Files="..."/>
 		</Node>
+		<Node Name="Upload To Cloud Mac" Requires="Make Installed Build Mac">
+			<CreateCloudArtifact Name="installed-build-mac" Type="installed-build" Description="Installed Build (Mac)" BaseDir="$(LocalInstalledDirMac)" Files="..." ProjectNameOverride="UE" Platform="Mac" Namespace="$(CloudPublishNamespace)" Host="$(CloudPublishHost)" BaselineBranch="$(CloudPublishBaseline)"/>
+		</Node>
 	</Agent>
 	<Aggregate Name="HostPlatforms_Mac" Label="Builds/Mac" Requires="Make Installed Build Mac"/>
 
@@ -1413,6 +1425,9 @@
 		<Node Name="Upload To Horde Linux" Requires="Make Installed Build Linux">
 			<CreateArtifact Name="installed-build-linux" Type="installed-build" Description="Installed Build (Linux)" BaseDir="$(LocalInstalledDirLinux)" Files="..."/>
 		</Node>
+		<Node Name="Upload To Cloud Linux" Requires="Make Installed Build Linux">
+			<CreateCloudArtifact Name="installed-build-linux" Type="installed-build" Description="Installed Build (Linux)" BaseDir="$(LocalInstalledDirLinux)" Files="..." ProjectNameOverride="UE" Platform="Linux" Namespace="$(CloudPublishNamespace)" Host="$(CloudPublishHost)" BaselineBranch="$(CloudPublishBaseline)"/>
+		</Node>
 	</Agent>
 	<Aggregate Name="HostPlatforms_Linux" Label="Builds/Linux" Requires="Make Installed Build Linux"/>
 
@@ -1420,6 +1435,11 @@
 	<Property Name="InstalledEngineBuildTargets" Value="$(InstalledEngineBuildTargets)Upload to Horde Win64;" If="'$(WithInstalledWin64)' == true"/>
 	<Property Name="InstalledEngineBuildTargets" Value="$(InstalledEngineBuildTargets)Upload to Horde Mac;" If="'$(WithInstalledMac)' == true"/>
 	<Property Name="InstalledEngineBuildTargets" Value="$(InstalledEngineBuildTargets)Upload to Horde Linux;" If="'$(WithInstalledLinux)' == true"/>
+
+	<Property Name="InstalledEngineBuildTargets" Value="$(InstalledEngineBuildTargets)Upload to Cloud Win64;" If="'$(WithInstalledWin64)' == true and $(PublishCloudArtifacts)"/>
+	<Property Name="InstalledEngineBuildTargets" Value="$(InstalledEngineBuildTargets)Upload to Cloud Mac;" If="'$(WithInstalledMac)' == true and $(PublishCloudArtifacts)"/>
+	<Property Name="InstalledEngineBuildTargets" Value="$(InstalledEngineBuildTargets)Upload to Cloud Linux;" If="'$(WithInstalledLinux)' == true and $(PublishCloudArtifacts)"/>
+	
 	<Aggregate Name="Installed Engine Build" Requires="$(InstalledEngineBuildTargets)"/>
 
 </BuildGraph>
diff --git a/Engine/Source/Programs/AutomationTool/BuildGraph/Tasks/CreateCloudArtifactTask.cs b/Engine/Source/Programs/AutomationTool/BuildGraph/Tasks/CreateCloudArtifactTask.cs
--- a/Engine/Source/Programs/AutomationTool/BuildGraph/Tasks/CreateCloudArtifactTask.cs
+++ b/Engine/Source/Programs/AutomationTool/BuildGraph/Tasks/CreateCloudArtifactTask.cs
@@ -58,11 +58,17 @@ namespace AutomationTool.Tasks
 		public string Platform { get; set; } = null!;
 
 		/// <summary>
-		/// The project this artifact is created for
+		/// The uproject this artifact is created for
 		/// </summary>
 		[TaskParameter(Optional = true)]
 		public FileReference Project { get; set; } = null!;
 
+		/// <summary>
+		/// Override the project name, defaults to the filename of the uproject
+		/// </summary>
+		[TaskParameter(Optional = true)]
+		public string ProjectNameOverride { get; set; } = "";
+
 		/// <summary>
 		/// The platform the artifact is created for
 		/// </summary>
@@ -178,8 +184,17 @@ namespace AutomationTool.Tasks
 			{
 				// if the project is not an uproject that exists we do the scanning for engine inis without the project specific overrides
 				projectPath = null;
-				Logger.LogInformation("\'{Project}\' is not a valid uproject file. Assuming it is a project name and not loading inis, you will need to specify project options explicitly.", _parameters.Project);
-				project = _parameters.Project.FullName;
+				Logger.LogInformation("\'{Project}\' is not a valid uproject file. Unable to load inis, you will need to specify project options explicitly.", _parameters.Project);
+				project = _parameters.ProjectNameOverride;
+				if (String.IsNullOrEmpty(project))
+				{
+					Logger.LogError("No project name override set and unable to locate uproject file '{Project}'. Unable to determine what the project name should be.", _parameters.Project);
+				}
+			}
+
+			if (!String.IsNullOrEmpty(_parameters.ProjectNameOverride))
+			{
+				project = _parameters.ProjectNameOverride;
 			}
 
 			bool foundConfig = false;
