Fixed issue were windows installed builds were uploaded as a different artifact type from the other platforms when uploading cloud artifacts

diff --git a/Engine/Build/InstalledEngineBuild.xml b/Engine/Build/InstalledEngineBuild.xml
--- a/Engine/Build/InstalledEngineBuild.xml
+++ b/Engine/Build/InstalledEngineBuild.xml
@@ -1169,7 +1169,7 @@
 			<CreateArtifact Name="installed-build-win64" Type="installed-build" Description="Installed Build (Windows)" BaseDir="$(LocalInstalledDir)" Files="..."/>
 		</Node>
 		<Node Name="Upload To Cloud Win64" Requires="Make Installed Build Win64">
-			<CreateCloudArtifact Name="installed-engine-win64" Type="installed-engine" Description="Installed Build (Windows)" BaseDir="$(LocalInstalledDir)" Files="..." ProjectNameOverride="UE"  Platform="Windows" Namespace="$(CloudPublishNamespace)" Host="$(CloudPublishHost)" BaselineBranch="$(CloudPublishBaseline)" />
+			<CreateCloudArtifact Name="installed-build-win64" Type="installed-build" Description="Installed Build (Windows)" BaseDir="$(LocalInstalledDir)" Files="..." ProjectNameOverride="UE"  Platform="Windows" Namespace="$(CloudPublishNamespace)" Host="$(CloudPublishHost)" BaselineBranch="$(CloudPublishBaseline)" />
 		</Node>
 	</Agent>
 	<Aggregate Name="HostPlatforms_Win64" Label="Builds/Win64" Requires="Make Installed Build Win64" Exclude="Update Version Files"/>
