Cloud Artifacts - Give the cloud artifacts a unique name compared to the horde artifacts - resolves issue with horde not liking conflicting names when artifacts are pre-registered

diff --git a/Engine/Build/InstalledEngineBuild.xml b/Engine/Build/InstalledEngineBuild.xml
--- a/Engine/Build/InstalledEngineBuild.xml
+++ b/Engine/Build/InstalledEngineBuild.xml
@@ -1169,7 +1169,7 @@
 			<CreateArtifact Name="installed-build-win64" Type="installed-build" Description="Installed Build (Windows)" BaseDir="$(LocalInstalledDir)" Files="..."/>
 		</Node>
 		<Node Name="Upload To Cloud Win64" Requires="Make Installed Build Win64">
-			<CreateCloudArtifact Name="installed-build-win64" Type="installed-build" Description="Installed Build (Windows)" BaseDir="$(LocalInstalledDir)" Files="..." ProjectNameOverride="UE"  Platform="Windows" Namespace="$(CloudPublishNamespace)" Host="$(CloudPublishHost)" BaselineBranch="$(CloudPublishBaseline)" />
+			<CreateCloudArtifact Name="cloud-installed-build-win64" Type="installed-build" Description="Installed Build (Windows)" BaseDir="$(LocalInstalledDir)" Files="..." ProjectNameOverride="UE"  Platform="Windows" Namespace="$(CloudPublishNamespace)" Host="$(CloudPublishHost)" BaselineBranch="$(CloudPublishBaseline)" />
 		</Node>
 	</Agent>
 	<Aggregate Name="HostPlatforms_Win64" Label="Builds/Win64" Requires="Make Installed Build Win64" Exclude="Update Version Files"/>
@@ -1303,7 +1303,7 @@
 			<CreateArtifact Name="installed-build-mac" Type="installed-build" Description="Installed Build (Mac)" BaseDir="$(LocalInstalledDirMac)" Files="..."/>
 		</Node>
 		<Node Name="Upload To Cloud Mac" Requires="Make Installed Build Mac">
-			<CreateCloudArtifact Name="installed-build-mac" Type="installed-build" Description="Installed Build (Mac)" BaseDir="$(LocalInstalledDirMac)" Files="..." ProjectNameOverride="UE" Platform="Mac" Namespace="$(CloudPublishNamespace)" Host="$(CloudPublishHost)" BaselineBranch="$(CloudPublishBaseline)"/>
+			<CreateCloudArtifact Name="cloud-installed-build-mac" Type="installed-build" Description="Installed Build (Mac)" BaseDir="$(LocalInstalledDirMac)" Files="..." ProjectNameOverride="UE" Platform="Mac" Namespace="$(CloudPublishNamespace)" Host="$(CloudPublishHost)" BaselineBranch="$(CloudPublishBaseline)"/>
 		</Node>
 	</Agent>
 	<Aggregate Name="HostPlatforms_Mac" Label="Builds/Mac" Requires="Make Installed Build Mac"/>
@@ -1426,7 +1426,7 @@
 			<CreateArtifact Name="installed-build-linux" Type="installed-build" Description="Installed Build (Linux)" BaseDir="$(LocalInstalledDirLinux)" Files="..."/>
 		</Node>
 		<Node Name="Upload To Cloud Linux" Requires="Make Installed Build Linux">
-			<CreateCloudArtifact Name="installed-build-linux" Type="installed-build" Description="Installed Build (Linux)" BaseDir="$(LocalInstalledDirLinux)" Files="..." ProjectNameOverride="UE" Platform="Linux" Namespace="$(CloudPublishNamespace)" Host="$(CloudPublishHost)" BaselineBranch="$(CloudPublishBaseline)"/>
+			<CreateCloudArtifact Name="cloud-installed-build-linux" Type="installed-build" Description="Installed Build (Linux)" BaseDir="$(LocalInstalledDirLinux)" Files="..." ProjectNameOverride="UE" Platform="Linux" Namespace="$(CloudPublishNamespace)" Host="$(CloudPublishHost)" BaselineBranch="$(CloudPublishBaseline)"/>
 		</Node>
 	</Agent>
 	<Aggregate Name="HostPlatforms_Linux" Label="Builds/Linux" Requires="Make Installed Build Linux"/>
