Removed localization files (.po, .archive, .manifest) from the engine install entirely
- Exception for `Engine.manifest` which is currently an engine dependency
- ~2.5GB install savings for `source code` chunk

#rnx

diff --git a/Engine/Build/InstalledEngineFilters.xml b/Engine/Build/InstalledEngineFilters.xml
--- a/Engine/Build/InstalledEngineFilters.xml
+++ b/Engine/Build/InstalledEngineFilters.xml
@@ -160,7 +160,9 @@
 		Engine/Plugins/Runtime/GameplayAbilities/README.md
 		Engine/Plugins/Experimental/RazerChromaDevices/README.txt
 
-	</Property>
+		<!-- Only standalone file required for Engine locs -->
+		Engine/Content/Localization/Engine/Engine.manifest
+</Property>
 	
 	<!-- Files which can exist under any engine or engine-platform directory -->
 	<Property Name="CopyEditorEngineOrPlatformFilter" Value="">
@@ -236,7 +238,7 @@
 
 		<!-- Exclude all the intermediate files in the Engine/Saved folder -->
 		Engine/Saved/...
-			<!-- Exclude generated config files/intermediates from Engine/Programs folder -->
+		<!-- Exclude generated config files/intermediates from Engine/Programs folder -->
 		Engine/Programs/AutomationTool/...
 		Engine/Programs/UnrealBuildTool/...
 		Engine/Programs/UnrealHeaderTool/Intermediate/...
@@ -251,6 +253,11 @@
 		Templates/*/Saved/...
 		Templates/*/manifest.json
 		Templates/*/contents.txt
+
+		<!-- Exclude locs files from the installed build entirely -->
+		Engine/Source/*/Localization/....archive
+		Engine/Source/*/Localization/....po
+		Engine/Source/*/Localization/....manifest
 	</Property>
 
 	<!-- Filters for installed DDC -->
