Ensure all sub paths are properly parsed when searching for localization files

#rnx

diff --git a/Engine/Build/InstalledEngineFilters.xml b/Engine/Build/InstalledEngineFilters.xml
--- a/Engine/Build/InstalledEngineFilters.xml
+++ b/Engine/Build/InstalledEngineFilters.xml
@@ -255,9 +255,9 @@
 		Templates/*/contents.txt
 
 		<!-- Exclude locs files from the installed build entirely -->
-		Engine/Source/*/Localization/....archive
-		Engine/Source/*/Localization/....po
-		Engine/Source/*/Localization/....manifest
+		Engine/Source/.../Localization/....archive
+		Engine/Source/.../Localization/....po
+		Engine/Source/.../Localization/....manifest
 	</Property>
 
 	<!-- Filters for installed DDC -->
