Make EditorPerfTests job upload metrics to PRS.

diff --git a/Engine/Build/EditorPerf.xml b/Engine/Build/EditorPerf.xml
--- a/Engine/Build/EditorPerf.xml
+++ b/Engine/Build/EditorPerf.xml
@@ -20,6 +20,7 @@
 	<!-- Projects setup-->
 	<Option Name="Projects" DefaultValue="false" Restrict=".+" Description="The projects to test"/>
 	<Error Message="No project names specified" If="'$(Projects)' == 'false'"/>
+	<Option Name="ProjectsToReportToPRS" DefaultValue="false" Restrict=".+" Description="The projects to upload test results to PRS"/>
 
 	<!-- Platform Options -->
 	<Option Name="Platforms" DefaultValue="Win64" Description="The list of platforms to profile"/>
@@ -200,6 +201,9 @@
 
 	<!-- options for summarize runs -->
 	<Option Name="RunTraceSummary" DefaultValue="false" Description="Summarize the insights trace after the run"/>
+	<Option Name="UploadTraceTimersToPRS" DefaultValue="true" Description="Export top timers from the utrace to PRS."/>
+	<Option Name="NumberOfTopTimers" DefaultValue="150" Description="The number of top (by inclusive time) timers to extract from the utrace for the PRS upload."/>
+	<Property Name="TimerRegionsToUpload" Value="Editor.Boot;Editor.Interact"/>
 
 	<!-- Gauntlet arguments for commandlet, summarize runs -->
 	<Property Name="CommandletArgs" Value="-notraceserver"/>
@@ -216,7 +220,7 @@
 	<Property Name="TestOutputDir" Value=""/>
 	<Property Name="RunTestLogPath" Value=""/>
 	<Property Name="TraceFilePath" Value=""/>
-	<Macro Name="RunTest" Arguments="TestNameSuffix;Platform;ProjectName;ProjectTarget;NodeName;Phase;Iteration;Args;EditorArgs">
+	<Macro Name="RunTest" Arguments="TestNameSuffix;Platform;ProjectName;ProjectTarget;NodeName;Phase;Iteration;Args;EditorArgs;Async;UploadToPRS">
 		<Property Name="TestName" Value="EditorPerf$(Phase)$(TestNameSuffix)"/>
 		<Property Name="AssetRegistryCacheRootFolderArg" Value="-AssetRegistryCacheRootFolder=&quot;$(AssetRegistryCache_Folder)&quot;"/>
 		
@@ -261,7 +265,14 @@
 
 		<!-- Run the Gauntlet EditorPerf Tests -->
 		<Command Name="RunUnreal" Arguments="$(AllGauntletArgs) $(ProfileArgs) $(RunArgs) -sessionlabel=&quot;$(TestName)&quot; -tracefile=&quot;$(TraceFilePath)&quot;"/>
-		
+
+		<Do If="$(UploadToPRS) AND ContainsItem('$(ProjectsToReportToPRS)', '$(Project)', ';')=='true'">			
+			<ForEach Name="TimerRegion" Values="$(TimerRegionsToUpload)">
+				<Property Name="TraceCSVFilePath" Value="$(TraceFilePath).$(TimerRegion).csv" />
+				<Command Name="UnrealInsights" Arguments="-OpenTraceFile=&quot;$(TraceFilePath)&quot; -autoquit -noui -ExecOnAnalysisCompleteCmd=&quot;TimingInsights.ExportTimerStatistics $(TraceCSVFilePath) -region=$(TimerRegion) -maxtimercount=$(NumberOfTopTimers) -sortby=totalinclusivetime -sortorder=descending&quot;"/>
+				<Command Name="UploadUnrealTimersToPRS" Arguments="-BuildVersion=$(BuildName) -Platform=Windows -Config=Development -Branch=$(Branch) -ExtraMetadata=&quot;Phase=$(Phase);Async=$(Async);&quot; -CL=$(Change) -DataSource=Automation.EditorPerfTests.$(ProjectName) -CSVFile=$(TraceCSVFilePath)"/>
+			</ForEach>
+		</Do>
 		<!-- Translate the utrace file into csv data --> 
 		<Commandlet Project="$(ProjectTarget)" Name="SummarizeTrace" Arguments="-inputfile=&quot;$(TraceFilePath)&quot; -testname=&quot;$(TestName)&quot; $(SummarizeArgs) $(CommandletArgs)" If="$(RunTraceSummary)"/>
 		
@@ -324,6 +335,7 @@
 				<Node Name="$(CompileNodeName)" Produces="$(SharedCompileProduces);$(CompileNodeNameProduces);$(TSANCompileNodeNameProduces)" IgnoreModified="$(ProjectBinariesPath)/...">
 					<SetVersion Change="$(Change)" Branch="$(EscapedBranch)" If="$(IsBuildMachine)"/>
 					<Compile Target="ShaderCompileWorker" Platform="$(Platform)" Configuration="Development" Tag="$(SharedCompileProduces)"/>
+					<Compile Target="UnrealInsights" Platform="$(Platform)" Configuration="Development" Tag="$(SharedCompileProduces)"/>
 					<Compile Target="CrashReportClientEditor" Platform="$(Platform)" Project="$(ProjectTarget)" Configuration="Shipping" Tag="$(SharedCompileProduces)"/>
 					<Compile Target="$(EditorTarget)" Platform="$(Platform)" Project="$(ProjectTarget)" Configuration="Development" Arguments="$(EditorCompileArgs)"/>
 					
@@ -349,7 +361,7 @@
 						<Node Name="$(RunNodeColdZenName)" Requires="$(CompileNodeName);$(SharedCompileProduces);$(CompileNodeNameProduces)" IgnoreModified="$(ProjectBinariesPath)/...">
 							<!-- Perform swap of built artifacts for TSAN builds so we use a non-TSAN build for cold runs and a TSAN build for warm runs -->
 							<Copy From="$(ProjectBinariesPath)_NonTSAN/..." To="$(ProjectBinariesPath)" ErrorIfNotFound="true" Overwrite="True"/>
-							<Expand Name="RunTest" TestNameSuffix="ZenLoader" Platform="$(Platform)" ProjectName="$(Project)" ProjectTarget="$(ProjectTarget)" NodeName="$(RunNodeColdZenName)" Phase="Cold" Iteration="1" Args="$(ColdArgs)" EditorArgs="$(ProjectCombinedArgs) $(ZenLoaderTestArgs)"/>
+							<Expand Name="RunTest" TestNameSuffix="ZenLoader" Platform="$(Platform)" ProjectName="$(Project)" ProjectTarget="$(ProjectTarget)" NodeName="$(RunNodeColdZenName)" Phase="Cold" Async="false" UploadToPRS="false" Iteration="1" Args="$(ColdArgs)" EditorArgs="$(ProjectCombinedArgs) $(ZenLoaderTestArgs)"/>
 						</Node>
 						<Property Name="RunNodeList" Value="$(RunNodeList);$(RunNodeColdZenName)"/>
 
@@ -357,7 +369,7 @@
 						<Property Name="RunTSANNodeZenName" Value="Run TSAN $(OutputName) Warm ZenLoader"/>
 						<Node Name="$(RunTSANNodeZenName)" Requires="$(RunNodeColdZenName)" IgnoreModified="$(ProjectBinariesPath)/...">
 							<Expand Name="CopyTSANBinariesIfNeeded" ProjectBinariesPath="$(ProjectBinariesPath)"/>
-							<Expand Name="RunTest" TestNameSuffix="ZenLoader" Platform="$(Platform)" ProjectName="$(Project)" ProjectTarget="$(ProjectTarget)" NodeName="$(RunTSANNodeZenName)" Phase="Warm" Iteration="1" Args="$(WarmArgs)" EditorArgs="$(ProjectCombinedArgs) $(ZenLoaderTestArgs)"/>
+							<Expand Name="RunTest" TestNameSuffix="ZenLoader" Platform="$(Platform)" ProjectName="$(Project)" ProjectTarget="$(ProjectTarget)" NodeName="$(RunTSANNodeZenName)" Phase="Warm" Async="false" UploadToPRS="false" Iteration="1" Args="$(WarmArgs)" EditorArgs="$(ProjectCombinedArgs) $(ZenLoaderTestArgs)"/>
 						</Node>
 						<Property Name="RunNodeList" Value="$(RunNodeList);$(RunTSANNodeZenName)"/>
 						<!-- Only depend on the Cold job since the node running TSAN likely will fail and we still want the report to run, even on failed jobs (e.g. it may have failed for TSAN reasons) -->
@@ -367,7 +379,7 @@
 						<Property Name="RunNodeTSANZenAsyncName" Value="Run TSAN $(OutputName) Warm ZenLoader Async"/>
 						<Node Name="$(RunNodeTSANZenAsyncName)" Requires="$(RunNodeColdZenName)" IgnoreModified="$(ProjectBinariesPath)/...">
 							<Expand Name="CopyTSANBinariesIfNeeded" ProjectBinariesPath="$(ProjectBinariesPath)"/>
-							<Expand Name="RunTest" TestNameSuffix="ZenLoaderAsync" Platform="$(Platform)" ProjectName="$(Project)" ProjectTarget="$(ProjectTarget)" NodeName="$(RunNodeTSANZenAsyncName)" Phase="Warm" Iteration="1" Args="$(WarmArgs)" EditorArgs="$(ProjectCombinedArgs) $(ZenLoaderAsyncTestArgs)"/>
+							<Expand Name="RunTest" TestNameSuffix="ZenLoaderAsync" Platform="$(Platform)" ProjectName="$(Project)" ProjectTarget="$(ProjectTarget)" NodeName="$(RunNodeTSANZenAsyncName)" Phase="Warm" Async="false" UploadToPRS="false" Iteration="1" Args="$(WarmArgs)" EditorArgs="$(ProjectCombinedArgs) $(ZenLoaderAsyncTestArgs)"/>
 						</Node>
 						<Property Name="RunNodeList" Value="$(RunNodeList);$(RunNodeTSANZenAsyncName)"/>
 						<!-- Only depend on the Cold job since the node running TSAN likely will fail and we still want the report to run, even on failed jobs (e.g. it may have failed for TSAN reasons) -->
@@ -393,7 +405,7 @@
 									<Property Name="RunNodeRequires" Value="$(RunNodeCompilationRequires)"/>
 									<Property Name="RunNodeRequires" Value="$(RunNodeCompilationRequires); Run $(OutputName) Cold$(IterationSuffix)" If="'$(Phase)' == 'Warm'"/>
 									<Node Name="$(RunNodeName)" Requires="$(RunNodeRequires)" IgnoreModified="$(ProjectBinariesPath)/...">
-										<Expand Name="RunTest" TestNameSuffix="NoZenLoader" Platform="$(Platform)" ProjectName="$(Project)" ProjectTarget="$(ProjectTarget)" NodeName="$(RunNodeName)" Phase="$(Phase)" Iteration="$(Iteration)" Args="$($(Phase)Args)" EditorArgs="$(ProjectCombinedArgs) $(NoZenLoaderTestArgs)"/>
+										<Expand Name="RunTest" TestNameSuffix="NoZenLoader" Platform="$(Platform)" ProjectName="$(Project)" ProjectTarget="$(ProjectTarget)" NodeName="$(RunNodeName)" Phase="$(Phase)" Async="false" UploadToPRS="false" Iteration="$(Iteration)" Args="$($(Phase)Args)" EditorArgs="$(ProjectCombinedArgs) $(NoZenLoaderTestArgs)"/>
 									</Node>
 									<Property Name="RunNodeList" Value="$(RunNodeList);$(RunNodeName)"/>
 								</ForEach>
@@ -408,7 +420,7 @@
 								<Property Name="RunNodeZenRequires" Value="$(RunNodeCompilationRequires)"/>
 								<Property Name="RunNodeZenRequires" Value="$(RunNodeCompilationRequires); Run $(OutputName) Cold ZenLoader" If="'$(Phase)' == 'Warm'"/>
 								<Node Name="$(RunNodeZenName)" Requires="$(RunNodeZenRequires)" IgnoreModified="$(ProjectBinariesPath)/...">
-									<Expand Name="RunTest" TestNameSuffix="ZenLoader" Platform="$(Platform)" ProjectName="$(Project)" ProjectTarget="$(ProjectTarget)" NodeName="$(RunNodeZenName)" Phase="$(Phase)" Iteration="1" Args="$($(Phase)Args)" EditorArgs="$(ProjectCombinedArgs) $(ZenLoaderTestArgs)"/>
+									<Expand Name="RunTest" TestNameSuffix="ZenLoader" Platform="$(Platform)" ProjectName="$(Project)" ProjectTarget="$(ProjectTarget)" NodeName="$(RunNodeZenName)" Phase="$(Phase)" Async="false" UploadToPRS="$(UploadTraceTimersToPRS)" Iteration="1" Args="$($(Phase)Args)" EditorArgs="$(ProjectCombinedArgs) $(ZenLoaderTestArgs)"/>
 								</Node>
 								<Property Name="RunNodeList" Value="$(RunNodeList);$(RunNodeZenName)" If="$(EnableZenLoaderTest)"/>
 							</ForEach>
@@ -421,7 +433,7 @@
 								<Property Name="RunNodeZenAsyncRequires" Value="$(RunNodeCompilationRequires)"/>
 								<Property Name="RunNodeZenAsyncRequires" Value="$(RunNodeCompilationRequires); Run $(OutputName) Cold ZenLoader Async" If="'$(Phase)' == 'Warm'"/>
 								<Node Name="$(RunNodeZenAsyncName)" Requires="$(RunNodeZenAsyncRequires)" IgnoreModified="$(ProjectBinariesPath)/...">
-									<Expand Name="RunTest" TestNameSuffix="ZenLoaderAsync" Platform="$(Platform)" ProjectName="$(Project)" ProjectTarget="$(ProjectTarget)" NodeName="$(RunNodeZenAsyncName)" Phase="$(Phase)" Iteration="1" Args="$($(Phase)Args)" EditorArgs="$(ProjectCombinedArgs) $(ZenLoaderAsyncTestArgs)"/>
+									<Expand Name="RunTest" TestNameSuffix="ZenLoaderAsync" Platform="$(Platform)" ProjectName="$(Project)" ProjectTarget="$(ProjectTarget)" NodeName="$(RunNodeZenAsyncName)" Phase="$(Phase)" Async="true" UploadToPRS="$(UploadTraceTimersToPRS)" Iteration="1" Args="$($(Phase)Args)" EditorArgs="$(ProjectCombinedArgs) $(ZenLoaderAsyncTestArgs)"/>
 								</Node>
 								<Property Name="RunNodeList" Value="$(RunNodeList);$(RunNodeZenAsyncName)" If="$(EnableZenLoaderAsyncThreadTest)"/>
 							</ForEach>
@@ -434,7 +446,7 @@
 								<Property Name="RunNodeZenStoreDDCRequires" Value="$(RunNodeCompilationRequires)"/>
 								<Property Name="RunNodeZenStoreDDCRequires" Value="$(RunNodeCompilationRequires); Run $(OutputName) Cold ZenLoader Store DDC" If="'$(Phase)' == 'Warm'"/>
 								<Node Name="$(RunNodeZenStoreDDCName)" Requires="$(RunNodeZenStoreDDCRequires)" IgnoreModified="$(ProjectBinariesPath)/...">
-									<Expand Name="RunTest" TestNameSuffix="ZenStoreAsyncDDC" Platform="$(Platform)" ProjectName="$(Project)" ProjectTarget="$(ProjectTarget)" NodeName="$(RunNodeZenStoreDDCName)" Phase="$(Phase)" Iteration="1" Args="$(ZenStoreDDCPhaseArgs)" EditorArgs="$(ProjectCombinedArgs) $(ZenStoreDDCTestArgs)"/>
+									<Expand Name="RunTest" TestNameSuffix="ZenStoreAsyncDDC" Platform="$(Platform)" ProjectName="$(Project)" ProjectTarget="$(ProjectTarget)" NodeName="$(RunNodeZenStoreDDCName)" Phase="$(Phase)" Async="false" UploadToPRS="false" Iteration="1" Args="$(ZenStoreDDCPhaseArgs)" EditorArgs="$(ProjectCombinedArgs) $(ZenStoreDDCTestArgs)"/>
 								</Node>
 								<Property Name="RunNodeList" Value="$(RunNodeList);$(RunNodeZenStoreDDCName)" If="$(EnableZenStoreDDCTest)"/>
 							</ForEach>
@@ -447,7 +459,7 @@
 							<Property Name="RunNodeLazyLoadRequires" Value="$(RunNodeCompilationRequires)"/>
 							<Property Name="RunNodeLazyLoadRequires" Value="$(RunNodeCompilationRequires); Run $(OutputName) Cold LazyLoad NoZenLoader" If="'$(Phase)' == 'Warm'"/>
 							<Node Name="$(RunNodeLazyLoadName)" Requires="$(RunNodeLazyLoadRequires)" IgnoreModified="$(ProjectBinariesPath)/..." If="$(EnableLazyLoadNoZenLoaderTest)">
-								<Expand Name="RunTest" TestNameSuffix="NoZenLoader" Platform="$(Platform)" ProjectName="$(Project)" ProjectTarget="$(ProjectTarget)" NodeName="$(RunNodeLazyLoadName)" Phase="$(Phase)" Iteration="1" Args="$($(Phase)Args)" EditorArgs="$(ProjectCombinedArgs) $(LazyLoadTestArgs) $(NoZenLoaderTestArgs)"/>
+								<Expand Name="RunTest" TestNameSuffix="NoZenLoader" Platform="$(Platform)" ProjectName="$(Project)" ProjectTarget="$(ProjectTarget)" NodeName="$(RunNodeLazyLoadName)" Phase="$(Phase)" Async="false" UploadToPRS="false" Iteration="1" Args="$($(Phase)Args)" EditorArgs="$(ProjectCombinedArgs) $(LazyLoadTestArgs) $(NoZenLoaderTestArgs)"/>
 							</Node>
 							<Property Name="RunNodeList" Value="$(RunNodeList);$(RunNodeLazyLoadName)" If="$(EnableLazyLoadNoZenLoaderTest)"/>
 
@@ -456,7 +468,7 @@
 							<Property Name="RunNodeLazyLoadZenRequires" Value="$(RunNodeCompilationRequires)"/>
 							<Property Name="RunNodeLazyLoadZenRequires" Value="$(RunNodeCompilationRequires); Run $(OutputName) Cold LazyLoad ZenLoader" If="'$(Phase)' == 'Warm'"/>
 							<Node Name="$(RunNodeLazyLoadZenName)" Requires="$(RunNodeLazyLoadZenRequires);" IgnoreModified="$(ProjectBinariesPath)/..." If="$(EnableLazyLoadZenLoaderTest)">
-								<Expand Name="RunTest" TestNameSuffix="ZenLoader" Platform="$(Platform)" ProjectName="$(Project)" ProjectTarget="$(ProjectTarget)" NodeName="$(RunNodeLazyLoadZenName)" Phase="$(Phase)" Iteration="1" Args="$($(Phase)Args)" EditorArgs="$(ProjectCombinedArgs) $(LazyLoadZenTestArgs)"/>
+								<Expand Name="RunTest" TestNameSuffix="ZenLoader" Platform="$(Platform)" ProjectName="$(Project)" ProjectTarget="$(ProjectTarget)" NodeName="$(RunNodeLazyLoadZenName)" Phase="$(Phase)" Async="false" UploadToPRS="false" Iteration="1" Args="$($(Phase)Args)" EditorArgs="$(ProjectCombinedArgs) $(LazyLoadZenTestArgs)"/>
 							</Node>
 							<Property Name="RunNodeList" Value="$(RunNodeList);$(RunNodeLazyLoadZenName)" If="$(EnableLazyLoadZenLoaderTest)"/>
 						</ForEach>
