Default option to include *.tgz docs should be false to avoid non-offline builds from bringing it into the `engine source` install chunk
- Saves ~250MB from the `engine source` install by ensuring these are only available in offline builds

#rnx

diff --git a/Engine/Build/InstalledEngineBuild.xml b/Engine/Build/InstalledEngineBuild.xml
--- a/Engine/Build/InstalledEngineBuild.xml
+++ b/Engine/Build/InstalledEngineBuild.xml
@@ -103,7 +103,7 @@
 	<Property Name="bAllowParallelExecutor" Value="$(AllowParallelExecutor)" />
 
 	<!-- Option to include source code docs in the resulting installed engine build -->
-	<Option Name="IncludeDocs" Restrict="true|false" DefaultValue="true" Description="Include the source code docs in the resulting installed engine build"/>
+	<Option Name="IncludeDocs" Restrict="true|false" DefaultValue="false" Description="Include the source code docs in the resulting installed engine build"/>
 
 	<!-- The local output directory -->
 	<Property Name="LocalInstalledDir" Value="$(BuiltDirectory)/Windows"/>
