More EditorPerf metrics (data source name, filters).

- Rename EditorPerf data source from to better match project name.
- Add ability to filter reports by DDC, AsyncLoad-ness and TimerRegion.

diff --git a/Engine/Build/EditorPerf.xml b/Engine/Build/EditorPerf.xml
--- a/Engine/Build/EditorPerf.xml
+++ b/Engine/Build/EditorPerf.xml
@@ -266,11 +266,11 @@
 		<!-- Run the Gauntlet EditorPerf Tests -->
 		<Command Name="RunUnreal" Arguments="$(AllGauntletArgs) $(ProfileArgs) $(RunArgs) -sessionlabel=&quot;$(TestName)&quot; -tracefile=&quot;$(TraceFilePath)&quot;"/>
 
-		<Do If="$(UploadToPRS) AND ContainsItem('$(ProjectsToReportToPRS)', '$(Project)', ';')=='true'">			
+		<Do If="$(UploadToPRS) AND ContainsItem('$(ProjectsToReportToPRS)', '$(Project)', ';')=='true'">
 			<ForEach Name="TimerRegion" Values="$(TimerRegionsToUpload)">
 				<Property Name="TraceCSVFilePath" Value="$(TraceFilePath).$(TimerRegion).csv" />
 				<Command Name="UnrealInsights" Arguments="-OpenTraceFile=&quot;$(TraceFilePath)&quot; -autoquit -noui -ExecOnAnalysisCompleteCmd=&quot;TimingInsights.ExportTimerStatistics $(TraceCSVFilePath) -region=$(TimerRegion) -maxtimercount=$(NumberOfTopTimers) -sortby=totalinclusivetime -sortorder=descending&quot;"/>
-				<Command Name="UploadUnrealTimersToPRS" Arguments="-BuildVersion=$(BuildName) -Platform=Windows -Config=Development -Branch=$(Branch) -ExtraMetadata=&quot;Phase=$(Phase);Async=$(Async);&quot; -CL=$(Change) -DataSource=Automation.EditorPerfTests.$(ProjectName) -CSVFile=$(TraceCSVFilePath)"/>
+				<Command Name="UploadUnrealTimersToPRS" Arguments="-BuildVersion=$(BuildName) -Platform=Windows -Config=Development -Branch=$(Branch) -ExtraMetadata=&quot;DDCPhase=$(Phase);AsyncLoad=$(Async);TimerRegion=$(TimerRegion)&quot; -CL=$(Change) -DataSource=Automation.$(ProjectName).EditorPerfTests -CSVFile=$(TraceCSVFilePath)"/>
 			</ForEach>
 		</Do>
 		<!-- Translate the utrace file into csv data --> 
