Add UAT command to wrap UnrealInsights export timers invocation.

- Also fix a number of more minor issues with EditorPerf script.

diff --git a/Engine/Build/EditorPerf.xml b/Engine/Build/EditorPerf.xml
--- a/Engine/Build/EditorPerf.xml
+++ b/Engine/Build/EditorPerf.xml
@@ -210,11 +210,12 @@
 	<Property Name="SummarizeArgs" Value=""/>
 
 	<!-- Publish telemetry options -->
+	<Property Name="JobURL" Value="$(HordeDashboardURL)/job/$(UE_HORDE_JOBID)?step=$(UE_HORDE_STEPID)" />
 	<Option Name="TelemetryConfig" DefaultValue="" Description="The configuration to use when publishing telemetry data to a database"/>
 	<Property Name="TelemetryConfig" Value="UETelemetryStaging" If="$(IsBuildMachine) And $(IsPreflight) And '$(TelemetryConfig)' == '' And $(RunTraceSummary)"/>
 	<Property Name="TelemetryConfig" Value="UETelemetry" If="$(IsBuildMachine) And '$(TelemetryConfig)' == '' And $(RunTraceSummary)"/>
 	<Option Name="ExtraTelemetryArgs" DefaultValue="" Description="Extra arguments passed to the publish telemetry command"/>
-	<Property Name="ExtraTelemetryArgs" Value="$(ExtraTelemetryArgs) -JobLink=&quot;$(HordeDashboardURL)/job/$(UE_HORDE_JOBID)?step=$(UE_HORDE_STEPID)&quot;" If="'$(HordeDashboardURL)' != ''"/>
+	<Property Name="ExtraTelemetryArgs" Value="$(ExtraTelemetryArgs) -JobLink=&quot;$(JobURL)&quot;" If="'$(HordeDashboardURL)' != ''"/>
 
 	<!-- The main test running macro logic that is expanded with platform and dependency node names -->
 	<Property Name="TestOutputDir" Value=""/>
@@ -227,6 +228,10 @@
 		<!-- Don't send telemetry for TSAN jobs as the time won't be comparable -->
 		<Property Name="EditorTelemetryArgs" Value=""/>
 		<Property Name="EditorTelemetryArgs" Value="-ST_SendTelemetry=false" If="$(WithTSAN)"/>
+
+		<!-- Don't itemize timers for TSAN jobs as the time won't be comparable -->
+		<Property Name="UploadToPRSEffective" Value="$(UploadToPRS)"/>
+		<Property Name="UploadToPRSEffective" Value="false" If="$(WithTSAN)"/>
 		
 		<Property Name="TestFilterArgs" Value=""/>
 		<Property Name="CombinedEditorArgs" Value=""/>
@@ -266,11 +271,11 @@
 		<!-- Run the Gauntlet EditorPerf Tests -->
 		<Command Name="RunUnreal" Arguments="$(AllGauntletArgs) $(ProfileArgs) $(RunArgs) -sessionlabel=&quot;$(TestName)&quot; -tracefile=&quot;$(TraceFilePath)&quot;"/>
 
-		<Do If="$(UploadToPRS) AND ContainsItem('$(ProjectsToReportToPRS)', '$(Project)', ';')=='true'">
+		<Do If="$(UploadToPRSEffective) AND ContainsItem('$(ProjectsToReportToPRS)', '$(Project)', ';')=='true'">
 			<ForEach Name="TimerRegion" Values="$(TimerRegionsToUpload)">
 				<Property Name="TraceCSVFilePath" Value="$(TraceFilePath).$(TimerRegion).csv" />
-				<Command Name="UnrealInsights" Arguments="-OpenTraceFile=&quot;$(TraceFilePath)&quot; -autoquit -noui -ExecOnAnalysisCompleteCmd=&quot;TimingInsights.ExportTimerStatistics $(TraceCSVFilePath) -region=$(TimerRegion) -maxtimercount=$(NumberOfTopTimers) -sortby=totalinclusivetime -sortorder=descending&quot;"/>
-				<Command Name="UploadUnrealTimersToPRS" Arguments="-BuildVersion=$(BuildName) -Platform=Windows -Config=Development -Branch=$(Branch) -ExtraMetadata=&quot;DDCPhase=$(Phase);AsyncLoad=$(Async);TimerRegion=$(TimerRegion)&quot; -CL=$(Change) -DataSource=Automation.$(ProjectName).EditorPerfTests -CSVFile=$(TraceCSVFilePath)"/>
+				<Command Name="ExportTimerStatisticsFromUtrace" Arguments="-TraceFile=&quot;$(TraceFilePath)&quot; -CSVFile=$(TraceCSVFilePath) -TimerRegion=$(TimerRegion) -MaxTimerCount=$(NumberOfTopTimers)"/>
+				<Command Name="UploadUnrealTimersToPRS" Arguments="-BuildVersion=$(BuildName) -Platform=Windows -Config=Development -Branch=$(Branch) -ExtraMetadata=&quot;DDCPhase=$(Phase);AsyncLoad=$(Async);TimerRegion=$(TimerRegion)&quot; -CL=$(Change) -DataSource=Automation.$(ProjectName).EditorPerfTests -CSVFile=$(TraceCSVFilePath) -JobURL=&quot;$(JobURL)&quot;"/>
 			</ForEach>
 		</Do>
 		<!-- Translate the utrace file into csv data --> 
diff --git a/Engine/Source/Programs/AutomationTool/Scripts/ExportTimerStatisticsFromUtrace.cs b/Engine/Source/Programs/AutomationTool/Scripts/ExportTimerStatisticsFromUtrace.cs
new file mode 100644
--- /dev/null
+++ b/Engine/Source/Programs/AutomationTool/Scripts/ExportTimerStatisticsFromUtrace.cs
@@ -0,0 +1,75 @@
+// Copyright Epic Games, Inc. All Rights Reserved.
+
+using AutomationTool;
+using EpicGames.Core;
+using System.Collections.Generic;
+using System.IO;
+using System;
+using UnrealBuildBase;
+using System.Linq;
+using System.Text;
+using UnrealBuildTool;
+using Polly;
+using Microsoft.Extensions.Logging;
+
+namespace ExportTimerStatisticsFromUtrace.Automation
+{
+	[Help("Runs UnrealInsights to export timers from a utrace file to CSV file.")]
+	[ParamHelp("TraceFile", "Input UnrealInsights .utrace file to open", ParamType = typeof(string), Required = true)]
+	[ParamHelp("CSVFile", "Output CSV file to write", ParamType = typeof(string), Required = true)]
+	[ParamHelp("TimerRegion", "TimerRegion to export", ParamType = typeof(string), Required = true)]
+	[ParamHelp("MaxTimerCount", "Number of top timers to export", ParamType = typeof(string), Required = true)]
+	class ExportTimerStatisticsFromUtrace : BuildCommand
+	{
+		public override void ExecuteBuild()
+		{
+			string TraceFile = ParseRequiredStringParam("TraceFile");
+			string CSVFile = ParseRequiredStringParam("CSVFile");
+			string TimerRegion = ParseRequiredStringParam("TimerRegion");
+			string NumTimers = ParseRequiredStringParam("MaxTimerCount");
+
+			string Arguments = string.Format("-OpenTraceFile=\"{0}\" -unattended -autoquit -noui -ExecOnAnalysisCompleteCmd=\"TimingInsights.ExportTimerStatistics {1} -region={2} -maxtimercount={3} -sortby=totalinclusivetime -sortorder=descending\"", 
+				TraceFile,
+				CSVFile,
+				TimerRegion,
+				NumTimers
+				);
+			Logger.LogInformation("About to run Insights with arguments '{Arguments}'", Arguments);
+			int ReturnCode = RunInsights(Arguments);
+			if (ReturnCode != 0)
+			{
+				throw new AutomationException("UnrealInsights invocation with arguments '{0}' failed with the return code {1}",
+					Arguments, ReturnCode);
+			}
+		}
+
+		/** Runs Insights with the given args and returns its error code. */
+		protected static int RunInsights(string Args, UnrealTargetConfiguration InConfiguration = UnrealTargetConfiguration.Development)
+		{
+			UnrealTargetPlatform ThisHostPlaform = HostPlatform.Platform;
+			string HostName = "UnrealInsights";
+			string Extension;
+			string BuildPath = Path.Combine(Unreal.EngineDirectory.FullName,"Binaries", ThisHostPlaform.ToString());
+			string BuildExecutable;
+			if (ThisHostPlaform != UnrealTargetPlatform.Mac)
+			{
+				Extension = ThisHostPlaform == UnrealTargetPlatform.Linux ? string.Empty : ".exe";
+				BuildExecutable = Path.Combine(BuildPath, InConfiguration == UnrealTargetConfiguration.Development ? $"{HostName}{Extension}" : $"{HostName}-{ThisHostPlaform}-{InConfiguration}{Extension}");
+			}
+			else 
+			{
+				Extension = ".app";
+				BuildExecutable = Path.Combine(BuildPath, "Contents", "MacOS",InConfiguration == UnrealTargetConfiguration.Development ? $"{HostName}" : $"{HostName}-{ThisHostPlaform}-{InConfiguration}");
+			}
+
+			if (!File.Exists(BuildExecutable))
+			{
+				throw new AutomationException("Was not able to find UnrealInsights binary at '{0}'. Make sure it is build in {1} configuration", 
+					BuildExecutable, InConfiguration.ToString());
+			}
+
+			Logger.LogInformation("Running binary '{BuildExecutable}' with arguments '{Args}'", BuildExecutable, Args);
+			return UnrealBuildTool.Utils.RunLocalProcessAndLogOutput(BuildExecutable, Args, EpicGames.Core.Log.Logger);
+		}
+	}
+}
