Exclude arm64 binaries from win64 base install
- Reduces base install size for win64 by ~500MB

#rnx

diff --git a/Engine/Build/InstalledEngineFilters.xml b/Engine/Build/InstalledEngineFilters.xml
--- a/Engine/Build/InstalledEngineFilters.xml
+++ b/Engine/Build/InstalledEngineFilters.xml
@@ -258,6 +258,12 @@
 		Engine/.../Localization/....archive
 		Engine/.../Localization/....po
 		Engine/.../Localization/....manifest
+
+		<!-- Exclude Windows-specific arm64 binaries as they aren't currently supported -->
+		Engine/.../win-arm64/...
+		Engine/.../WinArm64/...
+		Engine/.../Win64/arm64/...
+		Engine/.../overlay-arm64-uwp/...
 	</Property>
 
 	<!-- Filters for installed DDC -->
