Add missing thirdparty library for Android (HWCPipe moved)
-Submitting on behalf of chris.babcock.
#android

diff --git a/Engine/Build/InstalledEngineFilters.xml b/Engine/Build/InstalledEngineFilters.xml
--- a/Engine/Build/InstalledEngineFilters.xml
+++ b/Engine/Build/InstalledEngineFilters.xml
@@ -609,7 +609,7 @@
 		Engine/Source/ThirdParty/GoogleARCore/...
 		Engine/Source/ThirdParty/GoogleGameSDK/...
 		Engine/Source/ThirdParty/heapprofd/...
-		Engine/Source/ThirdParty/HWCPipe/...
+		Engine/Source/ThirdParty/libGPUCounters/...
 		Engine/Source/ThirdParty/Oculus/LibOVRAudio/LibOVRAudio/lib/*/*.so
 		Engine/Source/ThirdParty/Oculus/LibOVRAudio/*.tps
 		Engine/Source/ThirdParty/Oculus/LibOVRAudio/*.xml
