Added the TechSoft binaries to the launcher's packaging.

#rnx

diff --git a/Engine/Build/InstalledEngineFilters.xml b/Engine/Build/InstalledEngineFilters.xml
--- a/Engine/Build/InstalledEngineFilters.xml
+++ b/Engine/Build/InstalledEngineFilters.xml
@@ -361,6 +361,9 @@
 		Engine/Source/ThirdParty/Catch2/Catch2.build.cs
 		Engine/Source/ThirdParty/Catch2/Catch2.tps
 		Engine/Source/ThirdParty/Catch2/Catch2Extras.Build.cs
+
+		<!-- TechSoft -->
+		Engine/Binaries/ThirdParty/TechSoft/Win64/...
 	</Property>
 
 	<Property Name="CopyEditorExceptionsWin64" Value="$(CopyEditorExceptions)">
@@ -501,6 +504,9 @@
 		Engine/Source/ThirdParty/Catch2/Catch2.build.cs
 		Engine/Source/ThirdParty/Catch2/Catch2.tps
 		Engine/Source/ThirdParty/Catch2/Catch2Extras.Build.cs
+
+		<!-- TechSoft -->
+		Engine/Binaries/ThirdParty/TechSoft/Linux/...
 	</Property>
 	<Property Name="CopyEditorExceptionsLinux" Value="$(CopyEditorExceptions)">
 		<!-- Exclude Windows/Mac binaries on Linux here -->
