Upload GameThread editor startup timers separately.

diff --git a/Engine/Build/EditorPerf.xml b/Engine/Build/EditorPerf.xml
--- a/Engine/Build/EditorPerf.xml
+++ b/Engine/Build/EditorPerf.xml
@@ -275,7 +275,9 @@
 			<ForEach Name="TimerRegion" Values="$(TimerRegionsToUpload)">
 				<Property Name="TraceCSVFilePath" Value="$(TraceFilePath).$(TimerRegion).csv" />
 				<Command Name="ExportTimerStatisticsFromUtrace" Arguments="-TraceFile=&quot;$(TraceFilePath)&quot; -CSVFile=$(TraceCSVFilePath) -TimerRegion=$(TimerRegion) -MaxTimerCount=$(NumberOfTopTimers)"/>
-				<Command Name="UploadUnrealTimersToPRS" Arguments="-BuildVersion=$(BuildName) -Platform=Windows -Config=Development -Branch=$(Branch) -ExtraMetadata=&quot;DDCPhase=$(Phase);AsyncLoad=$(Async);TimerRegion=$(TimerRegion)&quot; -CL=$(Change) -DataSource=Automation.$(ProjectName).EditorPerfTests -CSVFile=$(TraceCSVFilePath) -JobURL=&quot;$(JobURL)&quot;"/>
+				<Command Name="UploadUnrealTimersToPRS" Arguments="-BuildVersion=$(BuildName) -Platform=Windows -Config=Development -Branch=$(Branch) -ExtraMetadata=&quot;DDCPhase=$(Phase);AsyncLoad=$(Async);TimerRegion=$(TimerRegion);Threads=All&quot; -CL=$(Change) -DataSource=Automation.$(ProjectName).EditorPerfTests -CSVFile=$(TraceCSVFilePath) -JobURL=&quot;$(JobURL)&quot;"/>
+				<Command Name="ExportTimerStatisticsFromUtrace" Arguments="-TraceFile=&quot;$(TraceFilePath)&quot; -CSVFile=$(TraceCSVFilePath) -TimerRegion=$(TimerRegion) -MaxTimerCount=$(NumberOfTopTimers) -Threads=GameThread"/>
+				<Command Name="UploadUnrealTimersToPRS" Arguments="-BuildVersion=$(BuildName) -Platform=Windows -Config=Development -Branch=$(Branch) -ExtraMetadata=&quot;DDCPhase=$(Phase);AsyncLoad=$(Async);TimerRegion=$(TimerRegion);Threads=GameThread&quot; -CL=$(Change) -DataSource=Automation.$(ProjectName).EditorPerfTests -CSVFile=$(TraceCSVFilePath) -JobURL=&quot;$(JobURL)&quot;"/>
 			</ForEach>
 		</Do>
 		<!-- Translate the utrace file into csv data --> 
