Android - package with ARCore fails at launch with installed editor
-Added the path for OpenXR_APL.xml, etc to CopyAndroidFilter.

diff --git a/Engine/Build/InstalledEngineFilters.xml b/Engine/Build/InstalledEngineFilters.xml
--- a/Engine/Build/InstalledEngineFilters.xml
+++ b/Engine/Build/InstalledEngineFilters.xml
@@ -604,6 +604,7 @@
 	<Property Name="CopyAndroidFilter">
 		<!-- Build Files -->
 		Engine/Binaries/DotNET/Android/...
+		Engine/Binaries/ThirdParty/OpenXR/Android/...
 		Engine/Build/Android/...
 		Engine/Extras/Android/...
 		Engine/Extras/LLDBDataFormatters/...
@@ -630,6 +631,8 @@
 		Engine/Source/ThirdParty/Oculus/OculusOpenXRLoader/*.tps
 		Engine/Source/ThirdParty/Oculus/OculusOpenXRLoader/OculusOpenXRLoader/*.xml
 		Engine/Source/ThirdParty/Oculus/OculusOpenXRLoader/OculusOpenXRLoader/Lib/*/*.so
+		Engine/Source/ThirdParty/OpenXR/*.xml
+		Engine/Source/ThirdParty/OpenXR/*.tps
 		Engine/Source/ThirdParty/ResonanceAudioApi/...
 		Engine/Source/ThirdParty/TangoSDK/...
 	</Property>
