Fix game thread timers stomping the previous CSV.

diff --git a/Engine/Build/EditorPerf.xml b/Engine/Build/EditorPerf.xml
--- a/Engine/Build/EditorPerf.xml
+++ b/Engine/Build/EditorPerf.xml
@@ -273,11 +273,13 @@
 
 		<Do If="$(UploadToPRSEffective) AND ContainsItem('$(ProjectsToReportToPRS)', '$(Project)', ';')=='true'">
 			<ForEach Name="TimerRegion" Values="$(TimerRegionsToUpload)">
-				<Property Name="TraceCSVFilePath" Value="$(TraceFilePath).$(TimerRegion).csv" />
-				<Command Name="ExportTimerStatisticsFromUtrace" Arguments="-TraceFile=&quot;$(TraceFilePath)&quot; -CSVFile=$(TraceCSVFilePath) -TimerRegion=$(TimerRegion) -MaxTimerCount=$(NumberOfTopTimers)"/>
-				<Command Name="UploadUnrealTimersToPRS" Arguments="-BuildVersion=$(BuildName) -Platform=Windows -Config=Development -Branch=$(Branch) -ExtraMetadata=&quot;DDCPhase=$(Phase);AsyncLoad=$(Async);TimerRegion=$(TimerRegion);Threads=All&quot; -CL=$(Change) -DataSource=Automation.$(ProjectName).EditorPerfTests -CSVFile=$(TraceCSVFilePath) -JobURL=&quot;$(JobURL)&quot;"/>
-				<Command Name="ExportTimerStatisticsFromUtrace" Arguments="-TraceFile=&quot;$(TraceFilePath)&quot; -CSVFile=$(TraceCSVFilePath) -TimerRegion=$(TimerRegion) -MaxTimerCount=$(NumberOfTopTimers) -Threads=GameThread"/>
-				<Command Name="UploadUnrealTimersToPRS" Arguments="-BuildVersion=$(BuildName) -Platform=Windows -Config=Development -Branch=$(Branch) -ExtraMetadata=&quot;DDCPhase=$(Phase);AsyncLoad=$(Async);TimerRegion=$(TimerRegion);Threads=GameThread&quot; -CL=$(Change) -DataSource=Automation.$(ProjectName).EditorPerfTests -CSVFile=$(TraceCSVFilePath) -JobURL=&quot;$(JobURL)&quot;"/>
+				<Property Name="TraceCSVFilePathBase" Value="$(TraceFilePath).$(TimerRegion)" />
+				<Property Name="TraceCSVFilePathAllThreads" Value="$(TraceCSVFilePathBase).AllThreads.csv" />
+				<Property Name="TraceCSVFilePathGameThread" Value="$(TraceCSVFilePathGameThread).GameThread.csv" />
+				<Command Name="ExportTimerStatisticsFromUtrace" Arguments="-TraceFile=&quot;$(TraceFilePath)&quot; -CSVFile=$(TraceCSVFilePathAllThreads) -TimerRegion=$(TimerRegion) -MaxTimerCount=$(NumberOfTopTimers)"/>
+				<Command Name="UploadUnrealTimersToPRS" Arguments="-BuildVersion=$(BuildName) -Platform=Windows -Config=Development -Branch=$(Branch) -ExtraMetadata=&quot;DDCPhase=$(Phase);AsyncLoad=$(Async);TimerRegion=$(TimerRegion);Threads=All&quot; -CL=$(Change) -DataSource=Automation.$(ProjectName).EditorPerfTests -CSVFile=$(TraceCSVFilePathAllThreads) -JobURL=&quot;$(JobURL)&quot;"/>
+				<Command Name="ExportTimerStatisticsFromUtrace" Arguments="-TraceFile=&quot;$(TraceFilePath)&quot; -CSVFile=$(TraceCSVFilePathGameThread) -TimerRegion=$(TimerRegion) -MaxTimerCount=$(NumberOfTopTimers) -Threads=GameThread"/>
+				<Command Name="UploadUnrealTimersToPRS" Arguments="-BuildVersion=$(BuildName) -Platform=Windows -Config=Development -Branch=$(Branch) -ExtraMetadata=&quot;DDCPhase=$(Phase);AsyncLoad=$(Async);TimerRegion=$(TimerRegion);Threads=GameThread&quot; -CL=$(Change) -DataSource=Automation.$(ProjectName).EditorPerfTests -CSVFile=$(TraceCSVFilePathGameThread) -JobURL=&quot;$(JobURL)&quot;"/>
 			</ForEach>
 		</Do>
 		<!-- Translate the utrace file into csv data --> 
