Fix a typo in EditorPerf scripts.

diff --git a/Engine/Build/EditorPerf.xml b/Engine/Build/EditorPerf.xml
--- a/Engine/Build/EditorPerf.xml
+++ b/Engine/Build/EditorPerf.xml
@@ -275,7 +275,7 @@
 			<ForEach Name="TimerRegion" Values="$(TimerRegionsToUpload)">
 				<Property Name="TraceCSVFilePathBase" Value="$(TraceFilePath).$(TimerRegion)" />
 				<Property Name="TraceCSVFilePathAllThreads" Value="$(TraceCSVFilePathBase).AllThreads.csv" />
-				<Property Name="TraceCSVFilePathGameThread" Value="$(TraceCSVFilePathGameThread).GameThread.csv" />
+				<Property Name="TraceCSVFilePathGameThread" Value="$(TraceCSVFilePathBase).GameThread.csv" />
 				<Command Name="ExportTimerStatisticsFromUtrace" Arguments="-TraceFile=&quot;$(TraceFilePath)&quot; -CSVFile=$(TraceCSVFilePathAllThreads) -TimerRegion=$(TimerRegion) -MaxTimerCount=$(NumberOfTopTimers)"/>
 				<Command Name="UploadUnrealTimersToPRS" Arguments="-BuildVersion=$(BuildName) -Platform=Windows -Config=Development -Branch=$(Branch) -ExtraMetadata=&quot;DDCPhase=$(Phase);AsyncLoad=$(Async);TimerRegion=$(TimerRegion);Threads=All&quot; -CL=$(Change) -DataSource=Automation.$(ProjectName).EditorPerfTests -CSVFile=$(TraceCSVFilePathAllThreads) -JobURL=&quot;$(JobURL)&quot;"/>
 				<Command Name="ExportTimerStatisticsFromUtrace" Arguments="-TraceFile=&quot;$(TraceFilePath)&quot; -CSVFile=$(TraceCSVFilePathGameThread) -TimerRegion=$(TimerRegion) -MaxTimerCount=$(NumberOfTopTimers) -Threads=GameThread"/>
