Force DDC misses on the MaterialShaderMap and GlobalShaderMap buckets when building the DDC Pak file for an installed engine build so that the individual shaders are queried and stored instead of shader maps.  With two-phase shader caching, this is necessary to ensure that the individual shaders are in the DDC Pak file.

diff --git a/Engine/Build/InstalledEngineBuild.xml b/Engine/Build/InstalledEngineBuild.xml
--- a/Engine/Build/InstalledEngineBuild.xml
+++ b/Engine/Build/InstalledEngineBuild.xml
@@ -884,8 +884,9 @@
 
 			<!-- Force enable D3D SM6 format; this is not enabled by default in engine settings but is used by various template content -->
 			<!-- As such we must ensure we compile engine shaders for SM6 as well -->
-			<Property Name="EngineExtraArgs" Value="&quot;-ini:Engine:[/Script/WindowsTargetPlatform.WindowsTargetSettings]:D3D12TargetedShaderFormats=PCD3D_SM6 $(ExtraDDCArgs)&quot;"/>
-			<Property Name="FeaturePackExtraArgs" Value="&quot;$(ExtraDDCArgs)&quot;"/>
+			<!-- Forcing individual shaders to be queried/stored (FShaderJobCacheShaders) instead of shader maps -->
+			<Property Name="EngineExtraArgs" Value="&quot;-ini:Engine:[/Script/WindowsTargetPlatform.WindowsTargetSettings]:D3D12TargetedShaderFormats=PCD3D_SM6 -DDC-All-MissTypes=MaterialShaderMap+GlobalShaderMap $(ExtraDDCArgs)&quot;"/>
+			<Property Name="FeaturePackExtraArgs" Value="&quot;-DDC-All-MissTypes=MaterialShaderMap+GlobalShaderMap $(ExtraDDCArgs)&quot;"/>
 
 			<Command Name="BuildDerivedDataCache" Arguments="-TempDir=&quot;$(DDCDir)&quot; -FeaturePacks=&quot;$(ProjectsToBuildDDCWin64)&quot; -TargetPlatforms=$(DDCPlatformsWin64) -HostPlatform=Win64 -SavedDir=&quot;$(SavedDir)&quot; -EngineContentExtraArgs=$(EngineExtraArgs) -FeaturePackExtraArgs=$(FeaturePackExtraArgs)"/>
 
@@ -936,8 +937,9 @@
 			<Copy From="$(CsToolsDir)" To="$(DDCDir)" Files="#Build Tools CS Binaries Mac"/>
 
 			<!-- Force disable trace starting for Mac -->
-			<Property Name="EngineExtraArgs" Value="-traceautostart=0"/>
-			<Property Name="FeaturePackExtraArgs" Value="-traceautostart=0"/>
+			<!-- Forcing individual shaders to be queried/stored (FShaderJobCacheShaders) instead of shader maps -->
+			<Property Name="EngineExtraArgs" Value="-DDC-All-MissTypes=MaterialShaderMap+GlobalShaderMap -traceautostart=0"/>
+			<Property Name="FeaturePackExtraArgs" Value="-DDC-All-MissTypes=MaterialShaderMap+GlobalShaderMap -traceautostart=0"/>
 
 			<Command Name="BuildDerivedDataCache" Arguments="-TempDir=&quot;$(DDCDir)&quot; -FeaturePacks=&quot;$(ProjectsToBuildDDC)&quot; -TargetPlatforms=$(DDCPlatformsMac) -HostPlatform=Mac -SavedDir=&quot;$(SavedDir)&quot; -EngineContentExtraArgs=$(EngineExtraArgs) -FeaturePackExtraArgs=$(FeaturePackExtraArgs)"/>
 
@@ -984,7 +986,11 @@
 			<Copy From="$(RootDir)" To="$(DDCDir)" Files="#FilteredCopyList"/>
 			<Copy From="$(CsToolsDir)" To="$(DDCDir)" Files="#Build Tools CS Binaries LinuxNative"/>
 
-			<Command Name="BuildDerivedDataCache" Arguments="-TempDir=&quot;$(DDCDir)&quot; -FeaturePacks=$(ProjectsToBuildDDC) -TargetPlatforms=$(DDCPlatformsLinux) -SavedDir=&quot;$(SavedDir)&quot;"/>
+			<!-- Forcing individual shaders to be queried/stored (FShaderJobCacheShaders) instead of shader maps -->
+			<Property Name="EngineExtraArgs" Value="-DDC-All-MissTypes=MaterialShaderMap+GlobalShaderMap"/>
+			<Property Name="FeaturePackExtraArgs" Value="-DDC-All-MissTypes=MaterialShaderMap+GlobalShaderMap"/>
+
+			<Command Name="BuildDerivedDataCache" Arguments="-TempDir=&quot;$(DDCDir)&quot; -FeaturePacks=$(ProjectsToBuildDDC) -TargetPlatforms=$(DDCPlatformsLinux) -SavedDir=&quot;$(SavedDir)&quot; -EngineContentExtraArgs=$(EngineExtraArgs) -FeaturePackExtraArgs=$(FeaturePackExtraArgs)"/>
 
 			<Tag Files="$(SavedDir)/Engine/DerivedDataCache/Compressed.ddp" With="#CompressedDDCLinux"/>
 		</Node>
