Fix installed build exceptions for cs build task not including intermediate UAT\UBT directories

#rnx

diff --git a/Engine/Build/InstalledEngineFilters.xml b/Engine/Build/InstalledEngineFilters.xml
--- a/Engine/Build/InstalledEngineFilters.xml
+++ b/Engine/Build/InstalledEngineFilters.xml
@@ -563,6 +563,18 @@
 		Engine/Source/Programs/Shared/*/obj/...
 	</Property>
 	<Property Name="CopyWin64CsToolsExceptions">
+		Engine/Saved/CsTools/Engine/Platforms/*/Source/Programs/AutomationTool/bin/...
+		Engine/Saved/CsTools/Engine/Platforms/*/Source/Programs/AutomationTool/obj/...
+		Engine/Saved/CsTools/Engine/Platforms/*/Source/Programs/AutomationTool/.../bin/...
+		Engine/Saved/CsTools/Engine/Platforms/*/Source/Programs/AutomationTool/.../obj/...
+		Engine/Saved/CsTools/Engine/Platforms/*/Source/Programs/UnrealBuildTool/bin/...
+		Engine/Saved/CsTools/Engine/Platforms/*/Source/Programs/UnrealBuildTool/obj/...
+		Engine/Saved/CsTools/Engine/Source/Programs/AutomationTool/bin/...
+		Engine/Saved/CsTools/Engine/Source/Programs/AutomationTool/obj/...
+		Engine/Saved/CsTools/Engine/Source/Programs/AutomationTool/.../bin/...
+		Engine/Saved/CsTools/Engine/Source/Programs/AutomationTool/.../obj/...
+		Engine/Saved/CsTools/Engine/Source/Programs/UnrealBuildTool/bin/...
+		Engine/Saved/CsTools/Engine/Source/Programs/UnrealBuildTool/obj/...
 		Engine/Saved/CsTools/Engine/Source/Programs/Shared/*/bin/...
 		Engine/Saved/CsTools/Engine/Source/Programs/Shared/*/obj/...
 	</Property>
@@ -592,6 +604,18 @@
 		Engine/Source/Programs/Shared/*/obj/...
 	</Property>
 	<Property Name="CopyMacCsToolsExceptions">
+		Engine/Saved/CsTools/Engine/Platforms/*/Source/Programs/AutomationTool/bin/...
+		Engine/Saved/CsTools/Engine/Platforms/*/Source/Programs/AutomationTool/obj/...
+		Engine/Saved/CsTools/Engine/Platforms/*/Source/Programs/AutomationTool/.../bin/...
+		Engine/Saved/CsTools/Engine/Platforms/*/Source/Programs/AutomationTool/.../obj/...
+		Engine/Saved/CsTools/Engine/Platforms/*/Source/Programs/UnrealBuildTool/bin/...
+		Engine/Saved/CsTools/Engine/Platforms/*/Source/Programs/UnrealBuildTool/obj/...
+		Engine/Saved/CsTools/Engine/Source/Programs/AutomationTool/bin/...
+		Engine/Saved/CsTools/Engine/Source/Programs/AutomationTool/obj/...
+		Engine/Saved/CsTools/Engine/Source/Programs/AutomationTool/.../bin/...
+		Engine/Saved/CsTools/Engine/Source/Programs/AutomationTool/.../obj/...
+		Engine/Saved/CsTools/Engine/Source/Programs/UnrealBuildTool/bin/...
+		Engine/Saved/CsTools/Engine/Source/Programs/UnrealBuildTool/obj/...
 		Engine/Saved/CsTools/Engine/Source/Programs/Shared/*/bin/...
 		Engine/Saved/CsTools/Engine/Source/Programs/Shared/*/obj/...
 	</Property>
@@ -760,6 +784,18 @@
 	</Property>
 
 	<Property Name="CopyLinuxCsToolsExceptions">
+		Engine/Saved/CsTools/Engine/Platforms/*/Source/Programs/AutomationTool/bin/...
+		Engine/Saved/CsTools/Engine/Platforms/*/Source/Programs/AutomationTool/obj/...
+		Engine/Saved/CsTools/Engine/Platforms/*/Source/Programs/AutomationTool/.../bin/...
+		Engine/Saved/CsTools/Engine/Platforms/*/Source/Programs/AutomationTool/.../obj/...
+		Engine/Saved/CsTools/Engine/Platforms/*/Source/Programs/UnrealBuildTool/bin/...
+		Engine/Saved/CsTools/Engine/Platforms/*/Source/Programs/UnrealBuildTool/obj/...
+		Engine/Saved/CsTools/Engine/Source/Programs/AutomationTool/bin/...
+		Engine/Saved/CsTools/Engine/Source/Programs/AutomationTool/obj/...
+		Engine/Saved/CsTools/Engine/Source/Programs/AutomationTool/.../bin/...
+		Engine/Saved/CsTools/Engine/Source/Programs/AutomationTool/.../obj/...
+		Engine/Saved/CsTools/Engine/Source/Programs/UnrealBuildTool/bin/...
+		Engine/Saved/CsTools/Engine/Source/Programs/UnrealBuildTool/obj/...
 		Engine/Saved/CsTools/Engine/Source/Programs/Shared/*/bin/...
 		Engine/Saved/CsTools/Engine/Source/Programs/Shared/*/obj/...
 	</Property>
