Force use of per-shader caching even for non-cook commandlets by using -forceAllowShaderCompilerJobCache commandline argument when running the DerivedDataCache commandlet to produce release editor derived data pak (DDP) files.

diff --git a/Engine/Build/InstalledEngineBuild.xml b/Engine/Build/InstalledEngineBuild.xml
--- a/Engine/Build/InstalledEngineBuild.xml
+++ b/Engine/Build/InstalledEngineBuild.xml
@@ -885,8 +885,8 @@
 			<!-- Force enable D3D SM6 format; this is not enabled by default in engine settings but is used by various template content -->
 			<!-- As such we must ensure we compile engine shaders for SM6 as well -->
 			<!-- Forcing individual shaders to be queried/stored (FShaderJobCacheShaders) instead of shader maps -->
-			<Property Name="EngineExtraArgs" Value="&quot;-ini:Engine:[/Script/WindowsTargetPlatform.WindowsTargetSettings]:D3D12TargetedShaderFormats=PCD3D_SM6 -DDC-All-MissTypes=MaterialShaderMap+GlobalShaderMap $(ExtraDDCArgs)&quot;"/>
-			<Property Name="FeaturePackExtraArgs" Value="&quot;-DDC-All-MissTypes=MaterialShaderMap+GlobalShaderMap $(ExtraDDCArgs)&quot;"/>
+			<Property Name="EngineExtraArgs" Value="&quot;-ini:Engine:[/Script/WindowsTargetPlatform.WindowsTargetSettings]:D3D12TargetedShaderFormats=PCD3D_SM6 -DDC-All-MissTypes=MaterialShaderMap+GlobalShaderMap -forceAllowShaderCompilerJobCache $(ExtraDDCArgs)&quot;"/>
+			<Property Name="FeaturePackExtraArgs" Value="&quot;-DDC-All-MissTypes=MaterialShaderMap+GlobalShaderMap -forceAllowShaderCompilerJobCache $(ExtraDDCArgs)&quot;"/>
 
 			<Command Name="BuildDerivedDataCache" Arguments="-TempDir=&quot;$(DDCDir)&quot; -FeaturePacks=&quot;$(ProjectsToBuildDDCWin64)&quot; -TargetPlatforms=$(DDCPlatformsWin64) -HostPlatform=Win64 -SavedDir=&quot;$(SavedDir)&quot; -EngineContentExtraArgs=$(EngineExtraArgs) -FeaturePackExtraArgs=$(FeaturePackExtraArgs)"/>
 
@@ -938,8 +938,8 @@
 
 			<!-- Force disable trace starting for Mac -->
 			<!-- Forcing individual shaders to be queried/stored (FShaderJobCacheShaders) instead of shader maps -->
-			<Property Name="EngineExtraArgs" Value="-DDC-All-MissTypes=MaterialShaderMap+GlobalShaderMap -traceautostart=0"/>
-			<Property Name="FeaturePackExtraArgs" Value="-DDC-All-MissTypes=MaterialShaderMap+GlobalShaderMap -traceautostart=0"/>
+			<Property Name="EngineExtraArgs" Value="-DDC-All-MissTypes=MaterialShaderMap+GlobalShaderMap -forceAllowShaderCompilerJobCache -traceautostart=0"/>
+			<Property Name="FeaturePackExtraArgs" Value="-DDC-All-MissTypes=MaterialShaderMap+GlobalShaderMap -forceAllowShaderCompilerJobCache -traceautostart=0"/>
 
 			<Command Name="BuildDerivedDataCache" Arguments="-TempDir=&quot;$(DDCDir)&quot; -FeaturePacks=&quot;$(ProjectsToBuildDDC)&quot; -TargetPlatforms=$(DDCPlatformsMac) -HostPlatform=Mac -SavedDir=&quot;$(SavedDir)&quot; -EngineContentExtraArgs=$(EngineExtraArgs) -FeaturePackExtraArgs=$(FeaturePackExtraArgs)"/>
 
@@ -987,8 +987,8 @@
 			<Copy From="$(CsToolsDir)" To="$(DDCDir)" Files="#Build Tools CS Binaries LinuxNative"/>
 
 			<!-- Forcing individual shaders to be queried/stored (FShaderJobCacheShaders) instead of shader maps -->
-			<Property Name="EngineExtraArgs" Value="-DDC-All-MissTypes=MaterialShaderMap+GlobalShaderMap"/>
-			<Property Name="FeaturePackExtraArgs" Value="-DDC-All-MissTypes=MaterialShaderMap+GlobalShaderMap"/>
+			<Property Name="EngineExtraArgs" Value="-DDC-All-MissTypes=MaterialShaderMap+GlobalShaderMap -forceAllowShaderCompilerJobCache"/>
+			<Property Name="FeaturePackExtraArgs" Value="-DDC-All-MissTypes=MaterialShaderMap+GlobalShaderMap -forceAllowShaderCompilerJobCache"/>
 
 			<Command Name="BuildDerivedDataCache" Arguments="-TempDir=&quot;$(DDCDir)&quot; -FeaturePacks=$(ProjectsToBuildDDC) -TargetPlatforms=$(DDCPlatformsLinux) -SavedDir=&quot;$(SavedDir)&quot; -EngineContentExtraArgs=$(EngineExtraArgs) -FeaturePackExtraArgs=$(FeaturePackExtraArgs)"/>
 
