Add missing enclosing quotes when passing extra commandline arguments to DerivedDataCache commandlet invocation on specifically Mac and Linux when producing installed editor builds.  Windows already had the enclosing quotes.

diff --git a/Engine/Build/InstalledEngineBuild.xml b/Engine/Build/InstalledEngineBuild.xml
--- a/Engine/Build/InstalledEngineBuild.xml
+++ b/Engine/Build/InstalledEngineBuild.xml
@@ -885,10 +885,10 @@
 			<!-- Force enable D3D SM6 format; this is not enabled by default in engine settings but is used by various template content -->
 			<!-- As such we must ensure we compile engine shaders for SM6 as well -->
 			<!-- Forcing individual shaders to be queried/stored (FShaderJobCacheShaders) instead of shader maps -->
-			<Property Name="EngineExtraArgs" Value="&quot;-ini:Engine:[/Script/WindowsTargetPlatform.WindowsTargetSettings]:D3D12TargetedShaderFormats=PCD3D_SM6 -DDC-All-MissTypes=MaterialShaderMap+GlobalShaderMap -forceAllowShaderCompilerJobCache $(ExtraDDCArgs)&quot;"/>
-			<Property Name="FeaturePackExtraArgs" Value="&quot;-DDC-All-MissTypes=MaterialShaderMap+GlobalShaderMap -forceAllowShaderCompilerJobCache $(ExtraDDCArgs)&quot;"/>
+			<Property Name="EngineExtraArgs" Value="-ini:Engine:[/Script/WindowsTargetPlatform.WindowsTargetSettings]:D3D12TargetedShaderFormats=PCD3D_SM6 -DDC-All-MissTypes=MaterialShaderMap+GlobalShaderMap -forceAllowShaderCompilerJobCache $(ExtraDDCArgs)"/>
+			<Property Name="FeaturePackExtraArgs" Value="-DDC-All-MissTypes=MaterialShaderMap+GlobalShaderMap -forceAllowShaderCompilerJobCache $(ExtraDDCArgs)"/>
 
-			<Command Name="BuildDerivedDataCache" Arguments="-TempDir=&quot;$(DDCDir)&quot; -FeaturePacks=&quot;$(ProjectsToBuildDDCWin64)&quot; -TargetPlatforms=$(DDCPlatformsWin64) -HostPlatform=Win64 -SavedDir=&quot;$(SavedDir)&quot; -EngineContentExtraArgs=$(EngineExtraArgs) -FeaturePackExtraArgs=$(FeaturePackExtraArgs)"/>
+			<Command Name="BuildDerivedDataCache" Arguments="-TempDir=&quot;$(DDCDir)&quot; -FeaturePacks=&quot;$(ProjectsToBuildDDCWin64)&quot; -TargetPlatforms=$(DDCPlatformsWin64) -HostPlatform=Win64 -SavedDir=&quot;$(SavedDir)&quot; -EngineContentExtraArgs=&quot;$(EngineExtraArgs)&quot; -FeaturePackExtraArgs=&quot;$(FeaturePackExtraArgs)&quot;"/>
 
 			<!-- Store off the generated pack csv file for later debugging -->
 			<!-- D:\build\++UE5\Sync\LocalBuilds\InstalledDDC\Engine\DerivedDataCache\ -->
@@ -941,7 +941,7 @@
 			<Property Name="EngineExtraArgs" Value="-DDC-All-MissTypes=MaterialShaderMap+GlobalShaderMap -forceAllowShaderCompilerJobCache -traceautostart=0"/>
 			<Property Name="FeaturePackExtraArgs" Value="-DDC-All-MissTypes=MaterialShaderMap+GlobalShaderMap -forceAllowShaderCompilerJobCache -traceautostart=0"/>
 
-			<Command Name="BuildDerivedDataCache" Arguments="-TempDir=&quot;$(DDCDir)&quot; -FeaturePacks=&quot;$(ProjectsToBuildDDC)&quot; -TargetPlatforms=$(DDCPlatformsMac) -HostPlatform=Mac -SavedDir=&quot;$(SavedDir)&quot; -EngineContentExtraArgs=$(EngineExtraArgs) -FeaturePackExtraArgs=$(FeaturePackExtraArgs)"/>
+			<Command Name="BuildDerivedDataCache" Arguments="-TempDir=&quot;$(DDCDir)&quot; -FeaturePacks=&quot;$(ProjectsToBuildDDC)&quot; -TargetPlatforms=$(DDCPlatformsMac) -HostPlatform=Mac -SavedDir=&quot;$(SavedDir)&quot; -EngineContentExtraArgs=&quot;$(EngineExtraArgs)&quot; -FeaturePackExtraArgs=&quot;$(FeaturePackExtraArgs)&quot;"/>
 
 			<Tag Files="$(SavedDir)/Engine/DerivedDataCache/Compressed.ddp" With="#CompressedDDCMac"/>
 		</Node>
@@ -990,7 +990,7 @@
 			<Property Name="EngineExtraArgs" Value="-DDC-All-MissTypes=MaterialShaderMap+GlobalShaderMap -forceAllowShaderCompilerJobCache"/>
 			<Property Name="FeaturePackExtraArgs" Value="-DDC-All-MissTypes=MaterialShaderMap+GlobalShaderMap -forceAllowShaderCompilerJobCache"/>
 
-			<Command Name="BuildDerivedDataCache" Arguments="-TempDir=&quot;$(DDCDir)&quot; -FeaturePacks=$(ProjectsToBuildDDC) -TargetPlatforms=$(DDCPlatformsLinux) -SavedDir=&quot;$(SavedDir)&quot; -EngineContentExtraArgs=$(EngineExtraArgs) -FeaturePackExtraArgs=$(FeaturePackExtraArgs)"/>
+			<Command Name="BuildDerivedDataCache" Arguments="-TempDir=&quot;$(DDCDir)&quot; -FeaturePacks=$(ProjectsToBuildDDC) -TargetPlatforms=$(DDCPlatformsLinux) -SavedDir=&quot;$(SavedDir)&quot; -EngineContentExtraArgs=&quot;$(EngineExtraArgs)&quot; -FeaturePackExtraArgs=&quot;$(FeaturePackExtraArgs)&quot;"/>
 
 			<Tag Files="$(SavedDir)/Engine/DerivedDataCache/Compressed.ddp" With="#CompressedDDCLinux"/>
 		</Node>
