Codesign AndroidZenServerPlugin binaries
#rnx

diff --git a/Engine/Build/InstalledEngineFilters.xml b/Engine/Build/InstalledEngineFilters.xml
--- a/Engine/Build/InstalledEngineFilters.xml
+++ b/Engine/Build/InstalledEngineFilters.xml
@@ -289,6 +289,8 @@
 		Engine/Binaries/Win64/EpicWebHelper.exe
 		Engine/Build/Windows/cl-filter/...
 		Engine/Extras/ThirdPartyNotUE/cwrsync/bin/...
+		<!-- Android ZenServerPlugin binaries -->
+		Engine/Source/Developer/Android/AndroidZenServerPlugin/Binaries/Win64/...
 	</Property>
 	
 	<Property Name="CopyEditorFilterWin64" Value="$(CopyEditorFilter)">
