Removed starter content references from feature pack/ddc build jobs and prevented it from being copied to final installed build package

#rnx

diff --git a/Engine/Build/InstalledEngineBuild.xml b/Engine/Build/InstalledEngineBuild.xml
--- a/Engine/Build/InstalledEngineBuild.xml
+++ b/Engine/Build/InstalledEngineBuild.xml
@@ -772,9 +772,6 @@
 			<ForEach Name="TemplateName" Values="$(ProjectsToFeaturePack)">
 				<PakFile Output="FeaturePacks/$(TemplateName).upack" Files="" ResponseFile="Templates/$(TemplateName)/contents.txt"/>
 			</ForEach>
-			<ForEach Name="SampleName" Values="StarterContent">
-				<PakFile Output="FeaturePacks/$(SampleName).upack" Files="" ResponseFile="Samples/$(SampleName)/contents.txt"/>
-			</ForEach>
 		</Node>
 	</Agent>
 
@@ -785,9 +782,6 @@
 			<ForEach Name="TemplateName" Values="$(ProjectsToFeaturePack)">
 				<PakFile Output="FeaturePacks/$(TemplateName).upack" Files="" ResponseFile="Templates/$(TemplateName)/contents.txt"/>
 			</ForEach>
-			<ForEach Name="SampleName" Values="StarterContent">
-				<PakFile Output="FeaturePacks/$(SampleName).upack" Files="" ResponseFile="Samples/$(SampleName)/contents.txt"/>
-			</ForEach>
 		</Node>
 	</Agent>
 
@@ -798,9 +792,6 @@
 			<ForEach Name="TemplateName" Values="$(ProjectsToFeaturePack)">
 				<PakFile Output="FeaturePacks/$(TemplateName).upack" Files="" ResponseFile="Templates/$(TemplateName)/contents.txt"/>
 			</ForEach>
-			<ForEach Name="SampleName" Values="StarterContent">
-				<PakFile Output="FeaturePacks/$(SampleName).upack" Files="" ResponseFile="Samples/$(SampleName)/contents.txt"/>
-			</ForEach>
 		</Node>
 	</Agent>
 
diff --git a/Engine/Build/InstalledEngineFilters.xml b/Engine/Build/InstalledEngineFilters.xml
--- a/Engine/Build/InstalledEngineFilters.xml
+++ b/Engine/Build/InstalledEngineFilters.xml
@@ -66,7 +66,6 @@
 		Templates/TP_SIM_BlankBP/TP_SIM_BlankBP.uproject
 		Templates/TP_DMXBP/TP_DMXBP.uproject
 		Templates/TP_MotionDesignBP/TP_MotionDesignBP.uproject
-		Samples/StarterContent/StarterContent.uproject
 		Engine/Source/Programs/LiveLinkHub/LiveLinkHub.uproject
 	</Property>
 	<!-- Win64 specific DDC -->
@@ -96,9 +95,6 @@
 		Engine/Source/UnrealGame.Target.cs
 		Engine/Source/UnrealEditor.Target.cs
 
-		<!-- Starter content -->
-		Samples/StarterContent/Content/...
-
 		<!-- Templates -->
 
 		Templates/TemplateCategories.ini
