Added UnrealSwarm tools back to installed engine build

#rnx

diff --git a/Engine/Build/InstalledEngineFilters.xml b/Engine/Build/InstalledEngineFilters.xml
--- a/Engine/Build/InstalledEngineFilters.xml
+++ b/Engine/Build/InstalledEngineFilters.xml
@@ -362,6 +362,11 @@
 
 		<!-- TechSoft -->
 		Engine/Binaries/ThirdParty/TechSoft/Win64/...
+
+		<!-- .NET Framework tools -->
+		Engine/Binaries/DotNET/SwarmCoordinator*
+		Engine/Binaries/DotNET/SwarmAgent*
+		Engine/Binaries/DotNET/SwarmInterface*
 	</Property>
 
 	<Property Name="CopyEditorExceptionsWin64" Value="$(CopyEditorExceptions)">
