Removing deleted ThirdParty/Unix dependencies from engine filters

#rnx

diff --git a/Engine/Build/InstalledEngineFilters.xml b/Engine/Build/InstalledEngineFilters.xml
--- a/Engine/Build/InstalledEngineFilters.xml
+++ b/Engine/Build/InstalledEngineFilters.xml
@@ -759,8 +759,6 @@
 		Engine/Binaries/ThirdParty/OpenAL/...
 		Engine/Binaries/ThirdParty/Vulkan/...
 		Engine/Binaries/ThirdParty/ICU/icu4c-53_1/Unix/x86_64-unknown-linux-gnu/*.so
-		Engine/Source/ThirdParty/Unix/LibCxx/...
-		Engine/Source/ThirdParty/Unix/HaveLinuxDependencies
 	</Property>
 
 	<Property Name="CopyLinuxFilterWin64" Value="$(CopyLinuxFilter)">
