Including additional missing UnrealSwarm dependencies to the win64 copy filter

#rnx

diff --git a/Engine/Build/InstalledEngineFilters.xml b/Engine/Build/InstalledEngineFilters.xml
--- a/Engine/Build/InstalledEngineFilters.xml
+++ b/Engine/Build/InstalledEngineFilters.xml
@@ -363,6 +363,9 @@
 		<!-- TechSoft -->
 		Engine/Binaries/ThirdParty/TechSoft/Win64/...
 
+		Engine/Binaries/DotNET/SwarmCommonUtils.dll
+		Engine/Binaries/DotNET/UnrealControls.dll
+		Engine/Binaries/DotNET/UnrealWindowsForms.deps.json
 	</Property>
 
 	<Property Name="CopyEditorExceptionsWin64" Value="$(CopyEditorExceptions)">
