ThirdParty: Catch2: Added support for building against multiple Standard Libraries on some platforms.

Also:
- Added missing .bat build scripts.
- Extended LowLevelTests.xml to allow platforms to build multiple variants of the library.

diff --git a/Engine/Build/LowLevelTests.xml b/Engine/Build/LowLevelTests.xml
--- a/Engine/Build/LowLevelTests.xml
+++ b/Engine/Build/LowLevelTests.xml
@@ -21,7 +21,9 @@
 	</Macro>	
 
 	<!-- Platform specific macros -->
-	<Macro Name="Catch2BuildAppendExtraCMakeArgsPlatform" Arguments="Platform"/>
+	<Macro Name="Catch2BuildVariantsPlatform" Arguments="Platform"/>
+	<Macro Name="Catch2BuildLibOutputPathPlatform" Arguments="Platform;Variant"/>
+	<Macro Name="Catch2BuildAppendExtraCMakeArgsPlatform" Arguments="Platform;Variant"/>
 	<Macro Name="Catch2BuildAppendExtraCMakeArgs" Arguments="ExtraArgValue">
 		<Property Name="CMakeAdditionalArgs" Value="$(CMakeAdditionalArgs) $(ExtraArgValue)"/>
 	</Macro>
@@ -259,57 +261,64 @@
 		</Agent>
 	</Macro>
 
+	<!-- All test metadata scripts along with platform specific versions for both public and internal paths -->
+	<Include Script="LowLevelTests/*.xml"/>
+	<Include Script="../Platforms/*/Build/LowLevelTests/*.xml"/>
+	<Include Script="../Restricted/NotForLicensees/Build/LowLevelTests/*.xml"/>
+	<Include Script="../Restricted/NotForLicensees/Platforms/*/Build/LowLevelTests/*.xml"/>
+	
 	<!-- Catch2 build targets -->
 	<ForEach Name="BuildPlatform" Values="$(BuildPlatforms)">
 		<ForEach Name="Platform" Values="$(PlatformsBuiltFrom$(BuildPlatform))">
 			<Property Name="AgentType" Value="$(BuildPlatform)"/>
 			<Expand Name="SelectBuildPlatformAndAgentType" Platform="$(Platform)" />
 			<Agent Name="Catch2 Build $(Platform)" Type="$(AgentType)">
-				<Node Name="Catch2 Build Library $(Platform)">
-					<Property Name="Generator" Value="Makefile"/>
-					<Property Name="TargetConfigs" Value="debug+release"/>
-					<Property Name="ExtraCMakeArgs" Value=""/>
-					<Property Name="CMakeAdditionalArgs" Value=""/>
-					<Property Name="TargetPlatform" Value="$(Platform)"/>
-					<Property Name="CommonExtraCmakeArgs" Value=""/>
-					<Property Name="CommonExtraCmakeArgs" Value="$(CommonExtraCmakeArgs) -DCATCH_DEVELOPMENT_BUILD=ON -DBUILD_TESTING=ON -DCATCH_BUILD_EXTRA_TESTS=ON -DCATCH_ENABLE_CONFIGURE_TESTS=ON" If="$(TestCatch2Libs)"/>
-					<Do If="'$(Platform)' == 'Android'">
-						<Property Name="ExtraCMakeArgs" Value="$(ExtraCMakeArgs) -TargetArchitecture=arm64"/>
-					</Do>
-					<Do If="'$(Platform)' == 'Linux'">
-						<Property Name="ExtraCMakeArgs" Value="$(ExtraCMakeArgs) -TargetArchitecture=x86_64-unknown-linux-gnu"/>
-					</Do>
-					<Do If="'$(Platform)' == 'LinuxArm64'">
-						<Property Name="TargetPlatform" Value="Linux"/>
-						<Property Name="ExtraCMakeArgs" Value="$(ExtraCMakeArgs) -TargetArchitecture=aarch64-unknown-linux-gnueabi"/>
-					</Do>
-					<Do If="'$(Platform)' == 'Win64'">
-						<Property Name="Generator" Value="VS2019"/>
-						<Do If="'$(Catch2LibVariation)' == 'VS2022'">
-							<Property Name="Generator" Value="VS2022"/>
+				<Property Name="Catch2BuildVariantsList" Value="Default"/>
+				<Expand Name="Catch2BuildVariantsPlatform" Platform="$(Platform)" />
+				
+				<ForEach Name="Variant" Values="$(Catch2BuildVariantsList)">				
+					<Node Name="Catch2 Build Library $(Platform) $(Variant)">
+						<Property Name="Generator" Value="Makefile"/>
+						<Property Name="TargetConfigs" Value="debug+release"/>
+						<Property Name="ExtraCMakeArgs" Value=""/>
+						<Property Name="CMakeAdditionalArgs" Value=""/>
+						<Property Name="LibOutputPath" Value="lib"/>
+						<Property Name="TargetPlatform" Value="$(Platform)"/>
+						<Property Name="CommonExtraCmakeArgs" Value=""/>
+						<Property Name="CommonExtraCmakeArgs" Value="$(CommonExtraCmakeArgs) -DCATCH_DEVELOPMENT_BUILD=ON -DBUILD_TESTING=ON -DCATCH_BUILD_EXTRA_TESTS=ON -DCATCH_ENABLE_CONFIGURE_TESTS=ON" If="$(TestCatch2Libs)"/>
+						<Do If="'$(Platform)' == 'Android'">
+							<Property Name="ExtraCMakeArgs" Value="$(ExtraCMakeArgs) -TargetArchitecture=arm64"/>
 						</Do>
-						<Property Name="ExtraCMakeArgs" Value="$(ExtraCMakeArgs) -TargetArchitecture=$(ArchitectureWin64)"/>
-					</Do>
-					<Expand Name="Catch2BuildAppendExtraCMakeArgsPlatform" Platform="$(Platform)"/>
-					<Do If="'$(BuildPlatform)' == 'Mac'">
-						<Property Name="Generator" Value="Xcode"/>
-						<Property Name="TargetConfigs" Value="Debug+Release"/>
-						<Property Name="CMakeAdditionalArgs" Value="$(CMakeAdditionalArgs) -DCMAKE_OSX_DEPLOYMENT_TARGET=10.15 -DCMAKE_OSX_ARCHITECTURES=arm64;x86_64"/>
-					</Do>
-					<Command Name="BuildCMakeLib"
-							 Arguments="-CMakeGenerator=$(Generator) -MakeTarget=Catch2 -TargetLib=Catch2 -TargetLibVersion=$(Catch2VersionBuild) -BinOutputPath=Binaries -LibOutputPath=lib -TargetPlatform=$(TargetPlatform) -TargetConfigs=$(TargetConfigs) -SkipCreateChangelist $(ExtraCMakeArgs) -CMakeAdditionalArguments=&quot;$(CommonExtraCmakeArgs) $(CMakeAdditionalArgs)&quot;"/>
-				</Node>
-				<Property Name="Catch2BuildNodes" Value="$(Catch2BuildNodes);Catch2 Build Library $(Platform)"/>
+						<Do If="'$(Platform)' == 'Linux'">
+							<Property Name="ExtraCMakeArgs" Value="$(ExtraCMakeArgs) -TargetArchitecture=x86_64-unknown-linux-gnu"/>
+						</Do>
+						<Do If="'$(Platform)' == 'LinuxArm64'">
+							<Property Name="TargetPlatform" Value="Linux"/>
+							<Property Name="ExtraCMakeArgs" Value="$(ExtraCMakeArgs) -TargetArchitecture=aarch64-unknown-linux-gnueabi"/>
+						</Do>
+						<Do If="'$(Platform)' == 'Win64'">
+							<Property Name="Generator" Value="VS2019"/>
+							<Do If="'$(Catch2LibVariation)' == 'VS2022'">
+								<Property Name="Generator" Value="VS2022"/>
+							</Do>
+							<Property Name="ExtraCMakeArgs" Value="$(ExtraCMakeArgs) -TargetArchitecture=$(ArchitectureWin64)"/>
+						</Do>
+						<Expand Name="Catch2BuildAppendExtraCMakeArgsPlatform" Platform="$(Platform)" Variant="$(Variant)"/>
+						<Expand Name="Catch2BuildLibOutputPathPlatform" Platform="$(Platform)" Variant="$(Variant)"/>
+						<Do If="'$(BuildPlatform)' == 'Mac'">
+							<Property Name="Generator" Value="Xcode"/>
+							<Property Name="TargetConfigs" Value="Debug+Release"/>
+							<Property Name="CMakeAdditionalArgs" Value="$(CMakeAdditionalArgs) -DCMAKE_OSX_DEPLOYMENT_TARGET=10.15 -DCMAKE_OSX_ARCHITECTURES=arm64;x86_64"/>
+						</Do>
+						<Command Name="BuildCMakeLib"
+								 Arguments="-CMakeGenerator=$(Generator) -MakeTarget=Catch2 -TargetLib=Catch2 -TargetLibVersion=$(Catch2VersionBuild) -BinOutputPath=Binaries -LibOutputPath=$(LibOutputPath) -TargetPlatform=$(TargetPlatform) -TargetConfigs=$(TargetConfigs) -SkipCreateChangelist $(ExtraCMakeArgs) -CMakeAdditionalArguments=&quot;$(CommonExtraCmakeArgs) $(CMakeAdditionalArgs)&quot;"/>
+					</Node>
+					<Property Name="Catch2BuildNodes" Value="$(Catch2BuildNodes);Catch2 Build Library $(Platform) $(Variant)"/>
+				</ForEach>
 			</Agent>
 		</ForEach>
 	</ForEach>
 
-	<!-- All test metadata scripts along with platform specific versions for both public and internal paths -->
-	<Include Script="LowLevelTests/*.xml"/>
-	<Include Script="../Platforms/*/Build/LowLevelTests/*.xml"/>
-	<Include Script="../Restricted/NotForLicensees/Build/LowLevelTests/*.xml"/>
-	<Include Script="../Restricted/NotForLicensees/Platforms/*/Build/LowLevelTests/*.xml"/>
-
 	<!-- If we set one of the Run*Tests option we don't want to run all the tests, just the selected ones. -->
 	<Property Name="RunsAllTests" Value="$(RunAllTests)"/>
 	<Expand Name="RunAllTests"/>
