Update bundled dotnet to 8.0.412

diff --git a/Engine/Build/BatchFiles/GetDotnetPath.bat b/Engine/Build/BatchFiles/GetDotnetPath.bat
--- a/Engine/Build/BatchFiles/GetDotnetPath.bat
+++ b/Engine/Build/BatchFiles/GetDotnetPath.bat
@@ -9,7 +9,7 @@ rem if UE_USE_SYSTEM_DOTNET we assume a installed dotnet is present
 if "%UE_USE_SYSTEM_DOTNET%" == "1" goto verify_dotnet
 
 rem if UE_DOTNET_VERSION is already set we assume this script has already run
-if "%UE_DOTNET_VERSION%" == "8.0.300" goto Succeeded
+if "%UE_DOTNET_VERSION%" == "8.0.412" goto Succeeded
 
 rem figure out which architecture to use
 set UE_DOTNET_ARCH=win-x64
@@ -22,7 +22,7 @@ if "%PROCESSOR_ARCHITECTURE%" == "ARM64" (
 )
 
 rem add the dotnet sdk in the repo as the current dotnet sdk
-set UE_DOTNET_VERSION=8.0.300
+set UE_DOTNET_VERSION=8.0.412
 set UE_DOTNET_DIR=%~dp0..\..\Binaries\ThirdParty\DotNet\%UE_DOTNET_VERSION%\%UE_DOTNET_ARCH%
 set PATH=%UE_DOTNET_DIR%;%PATH%
 set DOTNET_ROOT=%UE_DOTNET_DIR%
diff --git a/Engine/Build/BatchFiles/Linux/SetupDotnet.sh b/Engine/Build/BatchFiles/Linux/SetupDotnet.sh
--- a/Engine/Build/BatchFiles/Linux/SetupDotnet.sh
+++ b/Engine/Build/BatchFiles/Linux/SetupDotnet.sh
@@ -44,7 +44,7 @@ if [ $IS_DOTNET_INSTALLED -eq 0 ]; then
 	ARCH=x64
 	[ $(uname -m) == "arm64" ] && ARCH=arm64 
 
-	export UE_DOTNET_DIR="$CUR_DIR/../../../Binaries/ThirdParty/DotNet/8.0.300/linux-$ARCH"
+	export UE_DOTNET_DIR="$CUR_DIR/../../../Binaries/ThirdParty/DotNet/8.0.412/linux-$ARCH"
 	chmod u+x "$UE_DOTNET_DIR/dotnet"
 	export PATH="$UE_DOTNET_DIR:$PATH"
 	export DOTNET_ROOT="$UE_DOTNET_DIR"
diff --git a/Engine/Build/BatchFiles/Mac/SetupDotnet.sh b/Engine/Build/BatchFiles/Mac/SetupDotnet.sh
--- a/Engine/Build/BatchFiles/Mac/SetupDotnet.sh
+++ b/Engine/Build/BatchFiles/Mac/SetupDotnet.sh
@@ -44,7 +44,7 @@ if [ $IS_DOTNET_INSTALLED -eq 0 ]; then
 	ARCH=x64
 	[ $(uname -m) == "arm64" ] && ARCH=arm64 
 	
-	export UE_DOTNET_DIR=$CUR_DIR/../../../Binaries/ThirdParty/DotNet/8.0.300/mac-$ARCH
+	export UE_DOTNET_DIR=$CUR_DIR/../../../Binaries/ThirdParty/DotNet/8.0.412/mac-$ARCH
 	chmod u+x "$UE_DOTNET_DIR/dotnet"
 	echo $UE_DOTNET_DIR
 	export PATH=$UE_DOTNET_DIR:$PATH
diff --git a/Engine/Build/InstalledEngineFilters.xml b/Engine/Build/InstalledEngineFilters.xml
--- a/Engine/Build/InstalledEngineFilters.xml
+++ b/Engine/Build/InstalledEngineFilters.xml
@@ -277,7 +277,7 @@
 	<Property Name="ExtraSignFilesWin64">
 		<!-- CsvTools binaries which are manually updated in P4 -->
 		Engine/Binaries/DotNET/CsvTools/...
-		Engine/Binaries/ThirdParty/DotNet/8.0.300/win-x64/...
+		Engine/Binaries/ThirdParty/DotNet/8.0.412/win-x64/...
 		Engine/Binaries/ThirdParty/Python3/...
 		Engine/Binaries/ThirdParty/PresentMon/Win64/PresentMon-1.8.0-x64.exe
 		Engine/Binaries/ThirdParty/svn/Win64/...
@@ -326,7 +326,7 @@
 		Templates/TP_InCamVFXBP/...
 
 		<!-- DotNet -->
-		Engine/Binaries/ThirdParty/DotNet/8.0.300/win-x64/...
+		Engine/Binaries/ThirdParty/DotNet/8.0.412/win-x64/...
 		Engine/Binaries/ThirdParty/VisualStudio/...
 
 		<!-- CsvTools binaries which are manually updated in P4 -->
@@ -403,8 +403,8 @@
 		Engine/Build/Xcode/...
 
 		<!-- DotNet -->
-		Engine/Binaries/ThirdParty/DotNet/8.0.300/mac-arm64/...
-		Engine/Binaries/ThirdParty/DotNet/8.0.300/mac-x64/...
+		Engine/Binaries/ThirdParty/DotNet/8.0.412/mac-arm64/...
+		Engine/Binaries/ThirdParty/DotNet/8.0.412/mac-x64/...
 		Engine/Binaries/ThirdParty/VisualStudio/...
 	</Property>
 
@@ -479,8 +479,8 @@
 		Engine/Binaries/Linux/UnrealBuildAccelerator/...
 
 		<!-- DotNet -->
-		Engine/Binaries/ThirdParty/DotNet/8.0.300/linux-x64/...
-		Engine/Binaries/ThirdParty/DotNet/8.0.300/linux-arm64/...
+		Engine/Binaries/ThirdParty/DotNet/8.0.412/linux-x64/...
+		Engine/Binaries/ThirdParty/DotNet/8.0.412/linux-arm64/...
 		Engine/Binaries/ThirdParty/VisualStudio/...
 
 		<!-- DotNet bundled Requirements -->
@@ -669,8 +669,8 @@
 
 	<!-- IOS -->
 	<Property Name="ExtraSignFilesIOS">
-		Engine/Binaries/ThirdParty/DotNet/8.0.300/mac-arm64/...
-		Engine/Binaries/ThirdParty/DotNet/8.0.300/mac-x64/...
+		Engine/Binaries/ThirdParty/DotNet/8.0.412/mac-arm64/...
+		Engine/Binaries/ThirdParty/DotNet/8.0.412/mac-x64/...
 		Engine/Extras/ThirdPartyNotUE/libimobiledevice/x64/...
 	</Property>
 
@@ -721,8 +721,8 @@
 
 	<!-- TVOS -->
 	<Property Name="ExtraSignFilesTVOS">
-		Engine/Binaries/ThirdParty/DotNet/8.0.300/mac-arm64/...
-		Engine/Binaries/ThirdParty/DotNet/8.0.300/mac-x64/...
+		Engine/Binaries/ThirdParty/DotNet/8.0.412/mac-arm64/...
+		Engine/Binaries/ThirdParty/DotNet/8.0.412/mac-x64/...
 		Engine/Extras/ThirdPartyNotUE/libimobiledevice/x64/...
 	</Property>
 
diff --git a/Engine/Build/Rsync/RsyncEngine.txt b/Engine/Build/Rsync/RsyncEngine.txt
--- a/Engine/Build/Rsync/RsyncEngine.txt
+++ b/Engine/Build/Rsync/RsyncEngine.txt
@@ -64,12 +64,9 @@ protect /Binaries/TVOS/
 # Binaries/ThirdParty, just DotNet and VisualStudio required to run UnrealBuildTool
 + /Binaries/ThirdParty/
 + /Binaries/ThirdParty/DotNet/
-+ /Binaries/ThirdParty/DotNet/6.0.302/
-+ /Binaries/ThirdParty/DotNet/6.0.302/mac**/
-+ /Binaries/ThirdParty/DotNet/6.0.302/mac**/**
-+ /Binaries/ThirdParty/DotNet/8.0.300/
-+ /Binaries/ThirdParty/DotNet/8.0.300/mac**/
-+ /Binaries/ThirdParty/DotNet/8.0.300/mac**/**
++ /Binaries/ThirdParty/DotNet/8.0.412/
++ /Binaries/ThirdParty/DotNet/8.0.412/mac**/
++ /Binaries/ThirdParty/DotNet/8.0.412/mac**/**
 - /Binaries/ThirdParty/DotNet/**
 + /Binaries/ThirdParty/VisualStudio/
 + /Binaries/ThirdParty/VisualStudio/**
diff --git a/Engine/Source/Programs/Horde/BuildHorde.xml b/Engine/Source/Programs/Horde/BuildHorde.xml
--- a/Engine/Source/Programs/Horde/BuildHorde.xml
+++ b/Engine/Source/Programs/Horde/BuildHorde.xml
@@ -28,7 +28,7 @@
 	<Property Name="InformationalVersion" Value="$(InformationalVersion)-PF-$(PreflightChange)" If="'$(PreflightChange)' != ''"/>
 	<Property Name="VersionArguments" Value="/p:Version=$(Version).0 /p:InformationalVersion=$(InformationalVersion)"/>
 
-	<Property Name="DotNetRootPath" Value="$(RootDir)/Engine/Binaries/ThirdParty/DotNet/8.0.300"/>
+	<Property Name="DotNetRootPath" Value="$(RootDir)/Engine/Binaries/ThirdParty/DotNet/8.0.412"/>
 	<Property Name="DotNetPath" Value=""/>
 	<Property Name="DotNetPath" Value="$(DotNetRootPath)/win-x64/dotnet.exe" If="$(HostPlatform) == 'Win64' And $(HostArchitecture) == 'X64'"/>
 	<Property Name="DotNetPath" Value="$(DotNetRootPath)/win-arm64/dotnet.exe" If="$(HostPlatform) == 'Win64' And $(HostArchitecture) == 'Arm64'"/>
diff --git a/Engine/Source/Programs/Shared/EpicGames.Build/Unreal.cs b/Engine/Source/Programs/Shared/EpicGames.Build/Unreal.cs
--- a/Engine/Source/Programs/Shared/EpicGames.Build/Unreal.cs
+++ b/Engine/Source/Programs/Shared/EpicGames.Build/Unreal.cs
@@ -119,7 +119,7 @@ namespace UnrealBuildBase
 			return UnrealBuildToolDllPath;
 		}
 
-		private static string DotnetVersionDirectory = "8.0.300";
+		private static string DotnetVersionDirectory = "8.0.412";
 
 		private static string FindRelativeDotnetDirectory(RuntimePlatform.Type HostPlatform)
 		{
diff --git a/Engine/Source/Programs/UnrealCloudDDC/BuildUnrealCloudDDC.xml b/Engine/Source/Programs/UnrealCloudDDC/BuildUnrealCloudDDC.xml
--- a/Engine/Source/Programs/UnrealCloudDDC/BuildUnrealCloudDDC.xml
+++ b/Engine/Source/Programs/UnrealCloudDDC/BuildUnrealCloudDDC.xml
@@ -9,8 +9,8 @@
 	<Property Name="InformationalVersion" Value="$(InformationalVersion)-PF-$(PreflightChange)" If="'$(PreflightChange)' != ''"/>
 
 	<Property Name="DotNetPath" Value=""/>
-	<Property Name="DotNetPath" Value="$(RootDir)/Engine/Binaries/ThirdParty/DotNet/8.0.300/win-x64/dotnet.exe" If="$(HostPlatform) == 'Win64'"/>
-	<Property Name="DotNetPath" Value="$(RootDir)/Engine/Binaries/ThirdParty/DotNet/8.0.300/linux-x64/dotnet" If="$(HostPlatform) == 'Linux'"/>
+	<Property Name="DotNetPath" Value="$(RootDir)/Engine/Binaries/ThirdParty/DotNet/8.0.412/win-x64/dotnet.exe" If="$(HostPlatform) == 'Win64'"/>
+	<Property Name="DotNetPath" Value="$(RootDir)/Engine/Binaries/ThirdParty/DotNet/8.0.412/linux-x64/dotnet" If="$(HostPlatform) == 'Linux'"/>
 
 	<Agent Name="Jupiter" Type="Linux">
 		<Property Name="ComposeFile" Value="$(RootDir)/Engine/Source/Programs/UnrealCloudDDC/Composes/docker-compose-tests.yml"/>
