Add vkquality to installed engine build

diff --git a/Engine/Build/InstalledEngineFilters.xml b/Engine/Build/InstalledEngineFilters.xml
--- a/Engine/Build/InstalledEngineFilters.xml
+++ b/Engine/Build/InstalledEngineFilters.xml
@@ -665,6 +665,7 @@
 		Engine/Source/ThirdParty/OpenXR/*.tps
 		Engine/Source/ThirdParty/ResonanceAudioApi/...
 		Engine/Source/ThirdParty/TangoSDK/...
+		Engine/Source/ThirdParty/vkquality/...
 	</Property>
 	
 	<Property Name="CopyAndroidExceptions">
