Added `DiaSdk` path to engine copy filter for Windows
#rnx

diff --git a/Engine/Build/InstalledEngineFilters.xml b/Engine/Build/InstalledEngineFilters.xml
--- a/Engine/Build/InstalledEngineFilters.xml
+++ b/Engine/Build/InstalledEngineFilters.xml
@@ -370,6 +370,9 @@
 		Engine/Binaries/DotNET/SwarmCommonUtils.dll
 		Engine/Binaries/DotNET/UnrealControls.dll
 		Engine/Binaries/DotNET/UnrealWindowsForms.deps.json
+
+		<!-- DiaSdk -->
+		Engine/Source/ThirdParty/DiaSdk/...
 	</Property>
 
 	<Property Name="CopyEditorExceptionsWin64" Value="$(CopyEditorExceptions)">
