Adding new TP_MFG_BlankBP template.

diff --git a/Engine/Build/InstalledEngineFilters.xml b/Engine/Build/InstalledEngineFilters.xml
--- a/Engine/Build/InstalledEngineFilters.xml
+++ b/Engine/Build/InstalledEngineFilters.xml
@@ -55,6 +55,7 @@
 		Templates/TP_FirstPersonBP/TP_FirstPersonBP.uproject
 		Templates/TP_HandheldARBP/TP_HandheldARBP.uproject
 		Templates/TP_AEC_HandheldARBP/TP_AEC_HandheldARBP.uproject
+		Templates/TP_MFG_BlankBP/TP_MFG_BlankBP.uproject
 		Templates/TP_MFG_HandheldARBP/TP_MFG_HandheldARBP.uproject
 		Templates/TP_MFG_CollabBP/TP_MFG_CollabBP.uproject
 		Templates/TP_MFG_ProdConfigBP/TP_MFG_ProdConfigBP.uproject
@@ -120,6 +121,7 @@
 		Templates/TP_PhotoStudioBP/...
 		Templates/TP_ME_BlankBP/...
 		Templates/TP_ME_VProdBP/...
+		Templates/TP_MFG_BlankBP/...
 
 		Templates/TP_AEC_CollabBP/...
 		Templates/TP_MFG_CollabBP/...
