Add UnrealFrontend to Linux Installed Engine Build

#rnx

diff --git a/Engine/Build/InstalledEngineBuild.xml b/Engine/Build/InstalledEngineBuild.xml
--- a/Engine/Build/InstalledEngineBuild.xml
+++ b/Engine/Build/InstalledEngineBuild.xml
@@ -717,6 +717,7 @@
 			<Compile Target="ShaderCompileWorker" Configuration="Development" Platform="Linux" Tag="#Build Tools LinuxNative" Arguments="$(ExtraCompileArgs)"/>
 			<Compile Target="EpicWebHelper" Configuration="Shipping" Platform="Linux" Tag="#Build Tools LinuxNative" Arguments="$(ExtraCompileArgs)"/>
 			<Compile Target="UnrealInsights" Configuration="Development" Platform="Linux" Tag="#Build Tools LinuxNative" Arguments="$(ExtraCompileArgs)"/>
+			<Compile Target="UnrealFrontend" Configuration="Development" Platform="Linux" Tag="#Build Tools LinuxNative" Arguments="$(ExtraCompileArgs)"/>
 			<Compile Target="UnrealLightmass" Configuration="Development" Platform="Linux" Tag="#Build Tools LinuxNative" Arguments="$(ExtraCompileArgs)"/>
 			<Compile Target="InterchangeWorker" Configuration="Development" Platform="Linux" Tag="#Build Tools LinuxNative" Arguments="$(ExtraCompileArgs)"/>
 			<Compile Target="UnrealPak" Configuration="Development" Platform="Linux" Tag="#Build Tools LinuxNative" Arguments="$(ExtraCompileArgs)"/>
