Adding `win-arm64` installer option to launcher builds
#rnx

diff --git a/Engine/Build/InstalledEngineFilters.xml b/Engine/Build/InstalledEngineFilters.xml
--- a/Engine/Build/InstalledEngineFilters.xml
+++ b/Engine/Build/InstalledEngineFilters.xml
@@ -257,12 +257,6 @@
 		Engine/.../Localization/....po
 		Engine/.../Localization/....manifest
 
-		<!-- Exclude Windows-specific arm64 binaries as they aren't currently supported -->
-		Engine/.../win-arm64/...
-		Engine/.../WinArm64/...
-		Engine/.../Win64/arm64/...
-		Engine/.../overlay-arm64-uwp/...
-
 		<!-- Exclude Test Configs -->
 		Engine/Config/Tests/...
 		Engine/.../Config/Tests/...
