Win64 arm64 and arm64ec architecture support added to InstalledEngineBuild.
- -set:WithWinArm64=true and -set:WithWinArm64ec=true will have arm64 and arm64ec, respectively, built in addition to x64 which is always built.

diff --git a/Engine/Build/InstalledEngineBuild.xml b/Engine/Build/InstalledEngineBuild.xml
--- a/Engine/Build/InstalledEngineBuild.xml
+++ b/Engine/Build/InstalledEngineBuild.xml
@@ -16,6 +16,10 @@
 	<Property Name="DefaultWithWin64" Value="true"/>
 	<Property Name="DefaultWithWin64" Value="false" If="$(HostPlatformEditorOnly) Or ($(HostPlatformOnly) And '$(HostPlatform)' != 'Win64')"/>
 
+	<Property Name="DefaultWithWinArm64" Value="false"/>
+
+	<Property Name="DefaultWithWinArm64ec" Value="false"/>
+
 	<Property Name="DefaultWithMac" Value="true"/>
 	<Property Name="DefaultWithMac" Value="false" If="$(HostPlatformEditorOnly) Or ($(HostPlatformOnly) And '$(HostPlatform)' != 'Mac')"/>
 
@@ -53,7 +57,9 @@
 	<Option Name="WithInstalledLinux" Restrict="true|false" DefaultValue="true" Description="Whether to include the Linux editor as part of a binary release"/>
 
 	<!-- Setup default target platforms for Installed build if they haven't been specified on the commandline -->
-	<Option Name="WithWin64" Restrict="true|false" DefaultValue="$(DefaultWithWin64)" Description="Include the Win64 target platform"/>
+	<Option Name="WithWin64" Restrict="true|false" DefaultValue="$(DefaultWithWin64)" Description="Include the Win64 target platform (aka x64)"/>
+	<Option Name="WithWinArm64" Restrict="true|false" DefaultValue="$(DefaultWithWinArm64)" Description="Include the WinArm64 target platform"/>
+	<Option Name="WithWinArm64ec" Restrict="true|false" DefaultValue="$(DefaultWithWinArm64ec)" Description="Include the WinArm64ec target platform"/>
 	<Option Name="WithMac" Restrict="true|false" DefaultValue="$(DefaultWithMac)" Description="Include the Mac target platform"/>
 	<Option Name="WithAndroid" Restrict="true|false" DefaultValue="$(DefaultWithPlatform)" Description="Include the Android target platform"/>
 	<Option Name="WithIOS" Restrict="true|false" DefaultValue="$(DefaultWithIOS)" Description="Include the iOS target platform"/>
@@ -122,8 +128,11 @@
 
 	<!-- Architectures that we build for Android. NOTE: This should match the architectures specified in BuildPluginCommand. -->
 	<Property Name="AndroidArchitectures" Value="arm64+x64"/>
-	
-	<Property Name="DefaultArchitectureCommandsWin64" Value=""/>
+
+	<!-- Architectures that we build for Windows.  We always build x64, we may build others. NOTE: This should match the architectures specified in BuildPluginCommand. -->
+	<Property Name="Win64Architectures" Value="x64"/>
+	<Property Name="Win64Architectures" Value="$(Win64Architectures)+arm64" If="$(WithWinArm64)"/>
+	<Property Name="Win64Architectures" Value="$(Win64Architectures)+arm64ec" If="$(WithWinArm64ec)"/>
 	
 	<!-- Compile flags for Windows Game binaries -->
 	<Property Name="TargetDebugInfoArg" Value="-nodebuginfo"/>
@@ -321,8 +330,8 @@
 		<ForEach Name="Target" Values="$(OptionalClientTarget);$(OptionalServerTarget)">
 			<Node Name="Compile $(Target) Win64" Requires="Update Version Files" Produces="#$(Target) Win64">
 				<ForEach Name="Configuration" Values="$(GameConfigurations)">
-					<Compile Target="$(Target)" Platform="Win64" Configuration="$(Configuration)" Tag="#$(Target) Win64" Arguments="-precompile -allmodules -nolink $(TargetDebugInfoArg) $(ExtraCompileArgs) $(Win64NoPCHArg)"/>
-					<Compile Target="$(Target)" Platform="Win64" Configuration="$(Configuration)" Tag="#$(Target) Win64" Arguments="-precompile $(TargetDebugInfoArg) $(ExtraCompileArgs) $(Win64NoPCHArg)" Clean="false"/>
+					<Compile Target="$(Target)" Platform="Win64" Configuration="$(Configuration)" Tag="#$(Target) Win64" Arguments="-precompile -architectures=$(Win64Architectures) -allmodules -nolink $(TargetDebugInfoArg) $(ExtraCompileArgs) $(Win64NoPCHArg)"/>
+					<Compile Target="$(Target)" Platform="Win64" Configuration="$(Configuration)" Tag="#$(Target) Win64" Arguments="-precompile -architectures=$(Win64Architectures) $(TargetDebugInfoArg) $(ExtraCompileArgs) $(Win64NoPCHArg)" Clean="false"/>
 					<Property Name="OptionalTargetTags" Value="$(OptionalTargetTags);#$(Target) Win64"/>
 				</ForEach>
 			</Node>
@@ -332,8 +341,8 @@
 				<Tag Files="$(Tag)" With="#UnrealGame Win64"/>
 			</ForEach>
 			<ForEach Name="Configuration" Values="$(GameConfigurations)">
-				<Compile Target="UnrealGame" Platform="Win64" Configuration="$(Configuration)" Tag="#UnrealGame Win64" AllowParallelExecutor="$(bAllowParallelExecutor)" Arguments="-precompile -allmodules -nolink $(TargetDebugInfoArg) $(ExtraCompileArgs) $(Win64NoPCHArg)"/>
-				<Compile Target="UnrealGame" Platform="Win64" Configuration="$(Configuration)" Tag="#UnrealGame Win64" AllowParallelExecutor="$(bAllowParallelExecutor)" Arguments="-precompile $(TargetDebugInfoArg) $(ExtraCompileArgs) $(Win64NoPCHArg)" Clean="false"/>
+				<Compile Target="UnrealGame" Platform="Win64" Configuration="$(Configuration)" Tag="#UnrealGame Win64" AllowParallelExecutor="$(bAllowParallelExecutor)" Arguments="-precompile -architectures=$(Win64Architectures) -allmodules -nolink $(TargetDebugInfoArg) $(ExtraCompileArgs) $(Win64NoPCHArg)"/>
+				<Compile Target="UnrealGame" Platform="Win64" Configuration="$(Configuration)" Tag="#UnrealGame Win64" AllowParallelExecutor="$(bAllowParallelExecutor)" Arguments="-precompile -architectures=$(Win64Architectures) $(TargetDebugInfoArg) $(ExtraCompileArgs) $(Win64NoPCHArg)" Clean="false"/>
 			</ForEach>
 
 			<Do If="$(EmbedSrcSrvInfo)">
@@ -1153,7 +1162,8 @@
 			<Property Name="ContentOnlyPlatforms" Value=""/>
 
 			<!-- Create command line arguments for any platforms that were built with additional architectures -->
-			<Property Name="ArchitectureCommands" Value="$(DefaultArchitectureCommandsWin64)"/>
+			<Property Name="ArchitectureCommands" Value=""/>			
+			<Property Name="ArchitectureCommands" Value="$(ArchitectureCommands) -Win64Architectures=$(Win64Architectures)" If="'$(WithWin64)' == true"/>
 			<Property Name="ArchitectureCommands" Value="$(ArchitectureCommands) -AndroidArchitectures=$(AndroidArchitectures)" If="'$(WithAndroid)' == true"/>
 
 			<!-- Finalize the build and then tag the text file that this process creates -->
@@ -1300,6 +1310,7 @@
 
 			<!-- Create command line arguments for any platforms that were built with additional architectures -->
 			<Property Name="ArchitectureCommands" Value=""/>
+			<Property Name="ArchitectureCommands" Value="$(ArchitectureCommands) -Win64Architectures=$(Win64Architectures)" If="'$(WithWin64)' == true"/>
 			<Property Name="ArchitectureCommands" Value="$(ArchitectureCommands) -AndroidArchitectures=$(AndroidArchitectures)" If="'$(WithAndroid)' == true"/>
 			<!-- Finalize the build and then tag the text file that this process creates -->
 			<Command Name="FinalizeInstalledBuild" Arguments="-OutputDir=$(LocalInstalledDirMac) -Platforms=$(Platforms) $(ArchitectureCommands) -AnalyticsTypeOverride=$(AnalyticsTypeOverride)"/>
