Adding Unreal Intro template.

diff --git a/Engine/Build/InstalledEngineFilters.xml b/Engine/Build/InstalledEngineFilters.xml
--- a/Engine/Build/InstalledEngineFilters.xml
+++ b/Engine/Build/InstalledEngineFilters.xml
@@ -67,6 +67,7 @@
 		Templates/TP_SIM_BlankBP/TP_SIM_BlankBP.uproject
 		Templates/TP_DMXBP/TP_DMXBP.uproject
 		Templates/TP_MotionDesignBP/TP_MotionDesignBP.uproject
+		Templates/TP_UEIntro_BP/UEIntro.uproject
 		Engine/Source/Programs/LiveLinkHub/LiveLinkHub.uproject
 	</Property>
 	<!-- Win64 specific DDC -->
@@ -100,6 +101,7 @@
 
 		Templates/TemplateCategories.ini
 		Templates/Media/...
+		Templates/TP_UEIntro_BP/...
 
 		<!-- Game Templates -->
 		Templates/TP_Blank/...
