fix for UAT plugins not being handled correctly for installed engine builds

#rnx

diff --git a/Engine/Build/InstalledEngineBuild.xml b/Engine/Build/InstalledEngineBuild.xml
--- a/Engine/Build/InstalledEngineBuild.xml
+++ b/Engine/Build/InstalledEngineBuild.xml
@@ -580,7 +580,7 @@
 			<!-- Copy Source and referenced libraries to a new location with Confidential folders removed-->
 			<Tag Files="$(AutomationToolPaths)" Filter="*.csproj" Except="$(ConfidentialExceptions)"  With="#UAT Projects"/>
 			<CsCompile Project="#UAT Projects" Configuration="Development" Platform="AnyCPU" TagReferences="#UATReferences" EnumerateOnly="true"/>
-			<Tag Files="Engine/Binaries/DotNET/...;Engine/Binaries/ThirdParty/Newtonsoft/...;Engine/Binaries/ThirdParty/IOS/...;Engine/Binaries/ThirdParty/VisualStudio/...;Engine/Source/Programs/...;Engine/Platforms/*/Source/Programs/...;Engine/Source/Editor/SwarmInterface/...;Engine/Intermediate/ScriptModules/...;#UATReferences" Except="$(ConfidentialExceptions)" With="#RedistUATSource"/>
+			<Tag Files="Engine/Binaries/DotNET/...;Engine/Binaries/ThirdParty/Newtonsoft/...;Engine/Binaries/ThirdParty/IOS/...;Engine/Binaries/ThirdParty/VisualStudio/...;Engine/Source/Programs/...;Engine/Platforms/*/Source/Programs/...;Engine/Plugins/.../Build/Scripts/...;Engine/Platforms/*/Plugins/.../Build/Scripts/...;Engine/Source/Editor/SwarmInterface/...;Engine/Intermediate/ScriptModules/...;#UATReferences" Except="$(ConfidentialExceptions)" With="#RedistUATSource"/>
 			<Expand Name="Platform_TagCsToolsRedistUATSource"/>
 			<Copy Files="#RedistUATSource" From="$(RootDir)" To="$(CsToolsDir)"/>
 
diff --git a/Engine/Build/InstalledEngineFilters.xml b/Engine/Build/InstalledEngineFilters.xml
--- a/Engine/Build/InstalledEngineFilters.xml
+++ b/Engine/Build/InstalledEngineFilters.xml
@@ -576,6 +576,8 @@
 		Engine/Saved/CsTools/Engine/Platforms/*/Source/Programs/AutomationTool/.../obj/...
 		Engine/Saved/CsTools/Engine/Platforms/*/Source/Programs/UnrealBuildTool/bin/...
 		Engine/Saved/CsTools/Engine/Platforms/*/Source/Programs/UnrealBuildTool/obj/...
+		Engine/Saved/CsTools/Engine/Platforms/*/Plugins/.../Build/Scripts/bin/...
+		Engine/Saved/CsTools/Engine/Platforms/*/Plugins/.../Build/Scripts/obj/...
 		Engine/Saved/CsTools/Engine/Source/Programs/AutomationTool/bin/...
 		Engine/Saved/CsTools/Engine/Source/Programs/AutomationTool/obj/...
 		Engine/Saved/CsTools/Engine/Source/Programs/AutomationTool/.../bin/...
@@ -584,6 +586,8 @@
 		Engine/Saved/CsTools/Engine/Source/Programs/UnrealBuildTool/obj/...
 		Engine/Saved/CsTools/Engine/Source/Programs/Shared/*/bin/...
 		Engine/Saved/CsTools/Engine/Source/Programs/Shared/*/obj/...
+		Engine/Saved/CsTools/Engine/Plugins/.../Build/Scripts/bin/...
+		Engine/Saved/CsTools/Engine/Plugins/.../Build/Scripts/obj/...
 	</Property>
 
 	<!-- Mac -->
@@ -617,6 +621,8 @@
 		Engine/Saved/CsTools/Engine/Platforms/*/Source/Programs/AutomationTool/.../obj/...
 		Engine/Saved/CsTools/Engine/Platforms/*/Source/Programs/UnrealBuildTool/bin/...
 		Engine/Saved/CsTools/Engine/Platforms/*/Source/Programs/UnrealBuildTool/obj/...
+		Engine/Saved/CsTools/Engine/Platforms/*/Plugins/.../Build/Scripts/bin/...
+		Engine/Saved/CsTools/Engine/Platforms/*/Plugins/.../Build/Scripts/obj/...
 		Engine/Saved/CsTools/Engine/Source/Programs/AutomationTool/bin/...
 		Engine/Saved/CsTools/Engine/Source/Programs/AutomationTool/obj/...
 		Engine/Saved/CsTools/Engine/Source/Programs/AutomationTool/.../bin/...
@@ -625,6 +631,8 @@
 		Engine/Saved/CsTools/Engine/Source/Programs/UnrealBuildTool/obj/...
 		Engine/Saved/CsTools/Engine/Source/Programs/Shared/*/bin/...
 		Engine/Saved/CsTools/Engine/Source/Programs/Shared/*/obj/...
+		Engine/Saved/CsTools/Engine/Plugins/.../Build/Scripts/bin/...
+		Engine/Saved/CsTools/Engine/Plugins/.../Build/Scripts/obj/...
 	</Property>
 
 	<!-- Android -->
@@ -796,6 +804,8 @@
 		Engine/Saved/CsTools/Engine/Platforms/*/Source/Programs/AutomationTool/.../obj/...
 		Engine/Saved/CsTools/Engine/Platforms/*/Source/Programs/UnrealBuildTool/bin/...
 		Engine/Saved/CsTools/Engine/Platforms/*/Source/Programs/UnrealBuildTool/obj/...
+		Engine/Saved/CsTools/Engine/Platforms/*/Plugins/.../Build/Scripts/bin/...
+		Engine/Saved/CsTools/Engine/Platforms/*/Plugins/.../Build/Scripts/obj/...
 		Engine/Saved/CsTools/Engine/Source/Programs/AutomationTool/bin/...
 		Engine/Saved/CsTools/Engine/Source/Programs/AutomationTool/obj/...
 		Engine/Saved/CsTools/Engine/Source/Programs/AutomationTool/.../bin/...
@@ -804,6 +814,9 @@
 		Engine/Saved/CsTools/Engine/Source/Programs/UnrealBuildTool/obj/...
 		Engine/Saved/CsTools/Engine/Source/Programs/Shared/*/bin/...
 		Engine/Saved/CsTools/Engine/Source/Programs/Shared/*/obj/...
+		Engine/Saved/CsTools/Engine/Plugins/.../Build/Scripts/bin/...
+		Engine/Saved/CsTools/Engine/Plugins/.../Build/Scripts/obj/...
+
 	</Property>
 
 	<Property Name="CopyLinuxFilter" Value="$(CopyLinuxFilter);Engine/Extras/ThirdPartyNotUE/putty/..."/>	
