WinARM64 split compile support + temp build fixes for missing SxSManifest
#rnx

diff --git a/Engine/Build/InstalledEngineBuild.xml b/Engine/Build/InstalledEngineBuild.xml
--- a/Engine/Build/InstalledEngineBuild.xml
+++ b/Engine/Build/InstalledEngineBuild.xml
@@ -341,8 +341,10 @@
 				<Tag Files="$(Tag)" With="#UnrealGame Win64"/>
 			</ForEach>
 			<ForEach Name="Configuration" Values="$(GameConfigurations)">
-				<Compile Target="UnrealGame" Platform="Win64" Configuration="$(Configuration)" Tag="#UnrealGame Win64" AllowParallelExecutor="$(bAllowParallelExecutor)" Arguments="-precompile -architectures=$(Win64Architectures) -allmodules -nolink $(TargetDebugInfoArg) $(ExtraCompileArgs) $(Win64NoPCHArg)"/>
-				<Compile Target="UnrealGame" Platform="Win64" Configuration="$(Configuration)" Tag="#UnrealGame Win64" AllowParallelExecutor="$(bAllowParallelExecutor)" Arguments="-precompile -architectures=$(Win64Architectures) $(TargetDebugInfoArg) $(ExtraCompileArgs) $(Win64NoPCHArg)" Clean="false"/>
+				<ForEach Name="Architecture" Values="$(Win64Architectures)" Separator="+">
+					<Compile Target="UnrealGame" Platform="Win64" Configuration="$(Configuration)" Tag="#UnrealGame Win64" AllowParallelExecutor="$(bAllowParallelExecutor)" Arguments="-precompile -architectures=$(Architecture) -allmodules -nolink $(TargetDebugInfoArg) $(ExtraCompileArgs) $(Win64NoPCHArg)"/>
+					<Compile Target="UnrealGame" Platform="Win64" Configuration="$(Configuration)" Tag="#UnrealGame Win64" AllowParallelExecutor="$(bAllowParallelExecutor)" Arguments="-precompile -architectures=$(Architecture) $(TargetDebugInfoArg) $(ExtraCompileArgs) $(Win64NoPCHArg)" Clean="false"/>
+				</ForEach>
 			</ForEach>
 
 			<Do If="$(EmbedSrcSrvInfo)">
diff --git a/Engine/Source/Programs/UnrealBuildTool/Platform/Windows/VCToolChain.cs b/Engine/Source/Programs/UnrealBuildTool/Platform/Windows/VCToolChain.cs
--- a/Engine/Source/Programs/UnrealBuildTool/Platform/Windows/VCToolChain.cs
+++ b/Engine/Source/Programs/UnrealBuildTool/Platform/Windows/VCToolChain.cs
@@ -3879,6 +3879,14 @@ namespace UnrealBuildTool
 					BuildProducts.Add(item.Key.ChangeExtension($".alt{item.Key.GetExtension()}"), item.Value);
 				}
 			}
+
+			if (Target.WindowsPlatform.ManifestFile != null && Binary.Type == UEBuildBinaryType.Executable && !Target.Architecture.bIsX64)
+			{
+				string AssemblyName = Target.Architecture.ToString().ToLower();
+				DirectoryReference BaseDir = Target.ProjectFile?.Directory ?? Unreal.EngineDirectory;
+				FileReference IntermediateManifest = FileReference.Combine(BaseDir, "Intermediate", "SxSManifest", AssemblyName, AssemblyName + ".manifest");
+				BuildProducts.TryAdd(IntermediateManifest, BuildProductType.BuildResource);
+			}
 		}
 	}
 }
