Remove WinArm64 chunking option on Mac + ensure WinArm64 src/bins from p4 aren't in the base install
#rnx

diff --git a/Engine/Build/InstalledEngineFilters.xml b/Engine/Build/InstalledEngineFilters.xml
--- a/Engine/Build/InstalledEngineFilters.xml
+++ b/Engine/Build/InstalledEngineFilters.xml
@@ -449,6 +449,12 @@
 		<!-- Don't want The Windows/Linux Python packages -->
 		Engine/Plugins/Experimental/PythonFoundationPackages/Content/Python/Lib/Win64/...
 		Engine/Plugins/Experimental/PythonFoundationPackages/Content/Python/Lib/Linux/...
+
+		<!-- Exclude any files related to WinArm64 -->
+		Engine/.../win-arm64/...
+		Engine/.../WinArm64/...
+		Engine/.../Win64/arm64/...
+		Engine/.../overlay-arm64-uwp/...
 	</Property>
 
 	<!-- Linux editor specific dependencies -->
