Fixing missing Swarm binaries in packaged build

diff --git a/Engine/Build/InstalledEngineFilters.xml b/Engine/Build/InstalledEngineFilters.xml
--- a/Engine/Build/InstalledEngineFilters.xml
+++ b/Engine/Build/InstalledEngineFilters.xml
@@ -365,6 +365,10 @@
 		<!-- TechSoft -->
 		Engine/Binaries/ThirdParty/TechSoft/Win64/...
 
+		<!-- .NET Framework tools -->
+		Engine/Binaries/DotNET/SwarmCoordinator*
+		Engine/Binaries/DotNET/SwarmAgent*
+		Engine/Binaries/DotNET/SwarmInterface*
 		Engine/Binaries/DotNET/SwarmCommonUtils.dll
 		Engine/Binaries/DotNET/UnrealControls.dll
 		Engine/Binaries/DotNET/UnrealWindowsForms.deps.json
