Adding core localization manifests back to core engine install

#tests none

diff --git a/Engine/Build/InstalledEngineFilters.xml b/Engine/Build/InstalledEngineFilters.xml
--- a/Engine/Build/InstalledEngineFilters.xml
+++ b/Engine/Build/InstalledEngineFilters.xml
@@ -159,9 +159,6 @@
 		Engine/Plugins/Experimental/MoverAnimNext/README.md
 		Engine/Plugins/Runtime/GameplayAbilities/README.md
 		Engine/Plugins/Experimental/RazerChromaDevices/README.txt
-
-		<!-- Only standalone file required for Engine locs -->
-		Engine/Content/Localization/Engine/Engine.manifest
 </Property>
 	
 	<!-- Files which can exist under any engine or engine-platform directory -->
@@ -257,7 +254,8 @@
 		<!-- Exclude locs files from the installed build entirely -->
 		Engine/.../Localization/....archive
 		Engine/.../Localization/....po
-		Engine/.../Localization/....manifest
+		Engine/Content/Localization/PropertyNames.manifest
+		Engine/Content/Localization/ToolTips.manifest
 
 		<!-- Exclude Test Configs -->
 		Engine/Config/Tests/...
