Fix missing shell script for Designed for iPad feature in InstalledBuild

#rnx

diff --git a/Engine/Build/InstalledEngineFilters.xml b/Engine/Build/InstalledEngineFilters.xml
--- a/Engine/Build/InstalledEngineFilters.xml
+++ b/Engine/Build/InstalledEngineFilters.xml
@@ -395,8 +395,8 @@
 		Engine/Build/BatchFiles/Mac/GenerateUniversalDSYM.sh
 		Engine/Build/BatchFiles/Mac/ThinApp.sh
 		Engine/Build/BatchFiles/Mac/UpdateVersionAfterBuild.sh
-		Engine/Build/BatchFiles/Mac/ThinApp.sh
-        Engine/Build/BatchFiles/Mac/SourceCodeLookup.sh
+		Engine/Build/BatchFiles/Mac/SourceCodeLookup.sh
+		Engine/Build/BatchFiles/Mac/MakeMacNative.sh
 		Engine/Build/BatchFiles/RunDotnet.sh
 		Engine/Build/BatchFiles/RunDotnetMSBuild.sh
 		Engine/Build/BatchFiles/RunUAT.command
