[RaceDetector/CIS]
- Add RaceDetector on CitySample for EditorPerf in CIS

#FailedValidations Preflight

diff --git a/Engine/Build/EditorPerf.xml b/Engine/Build/EditorPerf.xml
--- a/Engine/Build/EditorPerf.xml
+++ b/Engine/Build/EditorPerf.xml
@@ -32,6 +32,7 @@
 	<!-- Build Options -->
 	<Option Name="NoDebugInfo" Restrict="true|false" DefaultValue="false" Description="Compile without including debug info"/>
 	<Option Name="WithTSAN" Restrict="true|false" DefaultValue="false" Description="Compile using TSAN"/>
+	<Option Name="WithRaceDetector" Restrict="true|false" DefaultValue="false" Description="Compile using Race Detector"/>
 	<Property Name="ExtraPluginsCommaDelims" Value="EditorTests"/>
 	<Property Name="ExtraPluginsPlusDelims" Value=""/> <!-- populated programmatically -->
 	<ForEach Name="Plugin" Values="$(ExtraPluginsCommaDelims)">
@@ -43,6 +44,7 @@
 	<Property Name="EditorCompileArgs" Value="-BuildPlugin=$(ExtraPluginsPlusDelims)" If="'$(ExtraPluginsPlusDelims)' != ''"/>
 	<Property Name="EditorCompileArgs" Value="$(EditorCompileArgs) -nodebuginfo" If="$(NoDebugInfo)"/>
 	<Property Name="TSANEditorCompileArgs" Value="-EnableTSAN -Monolithic"/>
+	<Property Name="RaceDetectorEditorCompileArgs" Value="-EnableInstrumentation"/>
 	<Property Name="EditorCompileArgs" Value="$(EditorCompileArgs) $(AdditionalEditorCompileArgs)" If="'$(AdditionalEditorCompileArgs)' != ''"/>
 
 	<!-- Get shared properties -->
@@ -72,7 +74,7 @@
 
 	<!-- Arguments for phases, keep in mind that these are passed to both the DDCCleanupCommandlet and RunUnreal -->
 	<Option Name="EnableBasePhaseTests" DefaultValue="true" Description="Enable the base phase perf tests."/>
-	<Property Name="EnableBasePhaseTests" Value="false" If="$(WithTSAN)"/>
+	<Property Name="EnableBasePhaseTests" Value="false" If="$(WithTSAN) OR $(WithRaceDetector)"/>
 	<Option Name="ColdArgs" DefaultValue="-ddc=Cold -NoAssetRegistryCacheRead" Description="Arguments to unreal for the cold phase"/>
 	<Option Name="WarmArgs" DefaultValue="-ddc=Warm" Description="Arguments to unreal for the warm phase"/>
 
@@ -82,7 +84,7 @@
 	<!-- Option for number of PIE iterations to run (per editor launch) -->
 	<Property Name="DefaultNumIterationsPie" Value="2" />
 	<Option Name="NumIterationsPIE" DefaultValue="$(DefaultNumIterationsPIE)" Description="Specify the number of times to start/stop PIE sessions in a single editor run."/>
-	<Property Name="NumIterationsPIE" Value="1" If="$(WithTSAN)"/>
+	<Property Name="NumIterationsPIE" Value="1" If="$(WithTSAN) OR $(WithRaceDetector)"/>
 	<Do If="'$(NumIterationsPIE)' &lt;= '0'">
 		<Warning Message="NumIterationPIE must be greater than 0. Setting to default of $(DefaultNumIterationsPIE)."/>
 		<Property Name="NumIterationsPIE" Value="$(DefaultNumIterationsPIE)" />
@@ -162,7 +164,8 @@
 	<!-- Zen Loader tests specifc arguments -->
 	<Property Name="ZenLoaderTestArgs" Value="-zenloader"/>
 	<Property Name="NoZenLoaderTestArgs" Value="-nozenloader"/>
-	<Property Name="ZenLoaderAsyncTestArgs" Value="$(ZenLoaderTestArgs) -asyncloadingthread" />
+	<Property Name="ZenLoaderAsyncTestArgs" Value="$(ZenLoaderTestArgs) -asyncloadingthread -dpcvars=s.AllowMultithreadedLoading=1" />
+	<Property Name="RaceDetectorTestArgs" Value="-racedetector" />
 	
 	<Property Name="ZenLoaderTestArgs" Value="$(ZenLoaderTestArgs) -noeditordomain" If="$(DisableEditorDomainForZenLoaderTests)" />
 	<Property Name="ZenLoaderAsyncTestArgs" Value="$(ZenLoaderAsyncTestArgs) -noeditordomain" If="$(DisableEditorDomainForZenLoaderTests)" />
@@ -221,6 +224,11 @@
 	<Property Name="TestOutputDir" Value=""/>
 	<Property Name="RunTestLogPath" Value=""/>
 	<Property Name="TraceFilePath" Value=""/>
+
+	<Macro Name="SetupTestOutputDir" Arguments="AutomationOutputDir;ProjectName;Platform;">
+		<Property Name="TestOutputDir" Value="$(AutomationOutputDir)/$(ProjectName)$(Platform)EditorPerf"/>
+	</Macro>
+
 	<Macro Name="RunTest" Arguments="TestNameSuffix;Platform;ProjectName;ProjectTarget;NodeName;Phase;Iteration;Args;EditorArgs;Async;UploadToPRS">
 		<Property Name="TestName" Value="EditorPerf$(Phase)$(TestNameSuffix)"/>
 		<Property Name="AssetRegistryCacheRootFolderArg" Value="-AssetRegistryCacheRootFolder=&quot;$(AssetRegistryCache_Folder)&quot;"/>
@@ -248,7 +256,7 @@
 		<Property Name="IterationSuffix" Value=""/>
 		<Property Name="IterationSuffix" Value="$(Iteration)" If="'$(Iteration)' != '1'"/>
 
-		<Property Name="TestOutputDir" Value="$(AutomationOutputDir)/$(ProjectName)$(Platform)EditorPerf"/>
+		<Expand Name="SetupTestOutputDir" AutomationOutputDir="$(AutomationOutputDir)" ProjectName="$(ProjectName)" Platform="$(Platform)"/>
 
 		<Property Name="IterationName" Value="$(ProjectName)$(Platform)$(TestName)$(IterationSuffix)"/>
 		<Property Name="TraceFileName" Value="$(IterationName).utrace"/>
@@ -288,11 +296,10 @@
 		<!-- Publish telemetry to the db -->
 		<Command Name="PublishUnrealAutomationTelemetry" Arguments="-CSVFile=&quot;$(TelemetryFilePath)&quot; -project=$(Project) -changelist=$(Change) -branch=$(Branch) -TelemetryConfig=$(TelemetryConfig) $(ExtraTelemetryArgs)" If="'$(TelemetryConfig)' != ''"/>
 	</Macro>
-	
+
 	<Property Name="GenerateTSANReportNodeName" Value=""/>
 	<Property Name="GenerateTSANReportNodeList" Value=""/>
-	<Macro Name="GenerateTSANReport" Arguments="NodeName;TestOutputDir;Phase;TestNameSuffix;Requires;IgnoreModified">
-		<Property Name="LogFilePath" Value="$(TestOutputDir)/$(Phase)$(TestNameSuffix)/Editor/EditorOutput.log"/>
+	<Macro Name="GenerateTSANReport" Arguments="NodeName;;Requires;IgnoreModified;LogFilePath">
 		<Property Name="GenerateTSANReportNodeName" Value="Sanitizer Report - $(NodeName)"/>
 		<Node Name="$(GenerateTSANReportNodeName)" Requires="$(Requires)" IgnoreModified="$(IgnoreModified)">
 			<Command Name="LogSanitizerOutput" Arguments="-logfile=&quot;$(LogFilePath)&quot;" />
@@ -309,7 +316,7 @@
 		<Copy From="$(ProjectBinariesPath)_TSAN/..." To="$(ProjectBinariesPath)" ErrorIfNotFound="true" Overwrite="True" If="$(NeedToCopyTSanBinaries)"/>
 		<Property Name="NeedToCopyTSanBinaries" Value="false"/>
 	</Macro>
-	
+
 
 	<!-- ForEach Platform to run, generate graph nodes -->
 	<ForEach Name="Platform" Values="$(Platforms)">
@@ -346,7 +353,8 @@
 					<Compile Target="ShaderCompileWorker" Platform="$(Platform)" Configuration="Development" Tag="$(SharedCompileProduces)"/>
 					<Compile Target="UnrealInsights" Platform="$(Platform)" Configuration="Development" Tag="$(SharedCompileProduces)"/>
 					<Compile Target="CrashReportClientEditor" Platform="$(Platform)" Project="$(ProjectTarget)" Configuration="Shipping" Tag="$(SharedCompileProduces)"/>
-					<Compile Target="$(EditorTarget)" Platform="$(Platform)" Project="$(ProjectTarget)" Configuration="Development" Arguments="$(EditorCompileArgs)"/>
+					<Compile Target="$(EditorTarget)" Platform="$(Platform)" Project="$(ProjectTarget)" Configuration="Development" Arguments="$(EditorCompileArgs)" If="!$(WithRaceDetector)"/>
+					<Compile Target="$(EditorTarget)" Platform="$(Platform)" Project="$(ProjectTarget)" Configuration="Development" Arguments="$(EditorCompileArgs) $(RaceDetectorEditorCompileArgs)" If="$(WithRaceDetector)"/>
 					
 					<Copy From="$(ProjectBinariesPath)/..." To="$(ProjectBinariesPath)_NonTSAN" Tag="$(CompileNodeNameProduces)" ErrorIfNotFound="true" If="$(WithTSAN)"/>
 					<Compile Target="$(EditorTarget)" Platform="$(Platform)" Project="$(ProjectTarget)" Configuration="Development" Arguments="$(EditorCompileArgs) $(TSANEditorCompileArgs)" If="$(WithTSAN)"/>
@@ -361,10 +369,34 @@
 			<Do If="ContainsItem('$(Projects)', '$(Project)', ';')=='true'">
 				<Property Name="RunNodeList" Value=""/>
 				<Property Name="UploadNodeList" Value=""/>
-				
+
 				<Agent Name="EditorPerf Agent $(Platform) $(Project)" Type="EditorPerf$(Platform)">
+					<!-- RaceDetector Nodes -->
+					<Do If="$(WithRaceDetector)">
+						<Expand Name="SetupTestOutputDir" AutomationOutputDir="$(AutomationOutputDir)" ProjectName="$(Project)" Platform="$(Platform)" />
+
+						<!-- Force a first Zen Loader run to prime the DDC for our actual test case (whatever it is, zen, legacy, lazy loader etc...) -->
+						<Property Name="RunNodeWarmupZenName" Value="Run Non-RaceDetector $(OutputName) Warmup ZenLoader Async"/>
+						<Node Name="$(RunNodeWarmupZenName)" Requires="$(CompileNodeName);$(SharedCompileProduces);$(CompileNodeNameProduces)" IgnoreModified="$(ProjectBinariesPath)/...">
+							<Expand Name="RunTest" TestNameSuffix="ZenLoader" Platform="$(Platform)" ProjectName="$(Project)" ProjectTarget="$(ProjectTarget)" NodeName="$(RunNodeWarmupZenName)" Phase="Warm" Async="false" UploadToPRS="false" Iteration="1" Args="$(WarmArgs)" EditorArgs="$(ProjectCombinedArgs) $(ZenLoaderAsyncTestArgs)"/>
+						</Node>
+						<Property Name="RunNodeList" Value="$(RunNodeList);$(RunNodeWarmupZenName)"/>
+
+						<!-- Zen Loader Async TestRun -->
+						<Property Name="RaceDetectorLogOutput" Value="$(TestOutputDir)/WarmZenLoaderAsync/Editor/RaceDetectorOutput.log"/>
+						<Property Name="RunNodeRaceDetectorZenAsyncName" Value="Run RaceDetector $(OutputName) Warm ZenLoader Async"/>
+						<Node Name="$(RunNodeRaceDetectorZenAsyncName)" Requires="$(CompileNodeName);$(SharedCompileProduces);$(CompileNodeNameProduces);$(RunNodeWarmupZenName)" IgnoreModified="$(ProjectBinariesPath)/...">
+							<Expand Name="RunTest" TestNameSuffix="ZenLoaderAsync" Platform="$(Platform)" ProjectName="$(Project)" ProjectTarget="$(ProjectTarget)" NodeName="$(RunNodeRaceDetectorZenAsyncName)" Phase="Warm" Async="false" UploadToPRS="false" Iteration="1" Args="$(WarmArgs)" EditorArgs="$(ProjectCombinedArgs) $(ZenLoaderAsyncTestArgs) $(RaceDetectorTestArgs) -racedetectorlog=$(RaceDetectorLogOutput)"/>
+						</Node>
+						<Property Name="RunNodeList" Value="$(RunNodeList);$(RunNodeRaceDetectorZenAsyncName)"/>
+						<!-- Only depend on the Cold job since the node running TSAN likely will fail and we still want the report to run, even on failed jobs (e.g. it may have failed for TSAN reasons) -->
+						<Expand Name="GenerateTSANReport" NodeName="$(RunNodeRaceDetectorZenAsyncName)" LogFilePath="$(RaceDetectorLogOutput)" IgnoreModified="$(ProjectBinariesPath)/..." Requires="$(RunNodeWarmupZenName)"/>
+					</Do>
+					
 					<!-- TSAN Nodes -->
 					<Do If="$(WithTSAN)">
+						<Expand Name="SetupTestOutputDir" AutomationOutputDir="$(AutomationOutputDir)" ProjectName="$(Project)" Platform="$(Platform)" />
+
 						<!-- Force a Cold Zen Loader run to prime the DDC for our actual test case (whatever it is, zen, legacy, lazy loader etc...) -->
 						<Property Name="RunNodeColdZenName" Value="Run Non-TSAN $(OutputName) Cold ZenLoader"/>
 						<Node Name="$(RunNodeColdZenName)" Requires="$(CompileNodeName);$(SharedCompileProduces);$(CompileNodeNameProduces)" IgnoreModified="$(ProjectBinariesPath)/...">
@@ -382,7 +414,8 @@
 						</Node>
 						<Property Name="RunNodeList" Value="$(RunNodeList);$(RunTSANNodeZenName)"/>
 						<!-- Only depend on the Cold job since the node running TSAN likely will fail and we still want the report to run, even on failed jobs (e.g. it may have failed for TSAN reasons) -->
-						<Expand Name="GenerateTSANReport" NodeName="$(RunTSANNodeZenName)" TestOutputDir="$(TestOutputDir)" Phase="Warm" TestNameSuffix="ZenLoader" IgnoreModified="$(ProjectBinariesPath)/..." Requires="$(RunNodeColdZenName)"/>
+						<Property Name="TSANLogOutput" Value="$(TestOutputDir)/WarmZenLoader/Editor/EditorOutput.log"/>
+						<Expand Name="GenerateTSANReport" NodeName="$(RunTSANNodeZenName)" LogFilePath="$(TSANLogOutput)" IgnoreModified="$(ProjectBinariesPath)/..." Requires="$(RunNodeColdZenName)"/>
 
 						<!-- Zen Loader Async TestRun -->
 						<Property Name="RunNodeTSANZenAsyncName" Value="Run TSAN $(OutputName) Warm ZenLoader Async"/>
@@ -392,11 +425,12 @@
 						</Node>
 						<Property Name="RunNodeList" Value="$(RunNodeList);$(RunNodeTSANZenAsyncName)"/>
 						<!-- Only depend on the Cold job since the node running TSAN likely will fail and we still want the report to run, even on failed jobs (e.g. it may have failed for TSAN reasons) -->
-						<Expand Name="GenerateTSANReport" NodeName="$(RunNodeTSANZenAsyncName)" TestOutputDir="$(TestOutputDir)" Phase="Warm" TestNameSuffix="ZenLoaderAsync" IgnoreModified="$(ProjectBinariesPath)/..." Requires="$(RunNodeColdZenName)"/>
+						<Property Name="TSANLogOutput" Value="$(TestOutputDir)/WarmZenLoaderAsync/Editor/EditorOutput.log"/>
+						<Expand Name="GenerateTSANReport" NodeName="$(RunNodeTSANZenAsyncName)" LogFilePath="$(TSANLogOutput)" IgnoreModified="$(ProjectBinariesPath)/..." Requires="$(RunNodeColdZenName)"/>
 					</Do>
 
 					<!-- Non-TSAN Nodes -->
-					<Do If="!$(WithTSAN)">
+					<Do If="!$(WithTSAN) AND !$(WithRaceDetector)">
 						<!-- 
 						Note, each of the nodes has a separate ForEach for Phases (where appropriate) because we want each node type (zen, nozen, async etc...) to be paired: Cold then Warm phase nodes.
 						If we did not order nodes like this, since the nodes run on the same agent it's possible for a cold job to fail and prevent a warm job _from an unrelated node_ to not have 
