Win64 arm64 and arm64ec
-Adding the -architectures parameter to the Compile UnrealGame Win64 node.
-Fix for arm64 include path.
-InstalledEngineBuild.xml was reviewed, in a release-5.7 cl.

diff --git a/Engine/Build/InstalledEngineBuild.xml b/Engine/Build/InstalledEngineBuild.xml
--- a/Engine/Build/InstalledEngineBuild.xml
+++ b/Engine/Build/InstalledEngineBuild.xml
@@ -131,8 +131,8 @@
 
 	<!-- Architectures that we build for Windows.  We always build x64, we may build others. NOTE: This should match the architectures specified in BuildPluginCommand. -->
 	<Property Name="Win64Architectures" Value="x64"/>
-	<Property Name="Win64Architectures" Value="$(Win64Architectures)+arm64" If="'$(WithWinArm64)' == true"/>
-	<Property Name="Win64Architectures" Value="$(Win64Architectures)+arm64ec" If="'$(WithWinArm64ec)' == true"/>	
+	<Property Name="Win64Architectures" Value="$(Win64Architectures)+arm64" If="$(WithWinArm64)"/>
+	<Property Name="Win64Architectures" Value="$(Win64Architectures)+arm64ec" If="$(WithWinArm64ec)"/>
 	
 	<!-- Compile flags for Windows Game binaries -->
 	<Property Name="TargetDebugInfoArg" Value="-nodebuginfo"/>
@@ -330,8 +330,8 @@
 		<ForEach Name="Target" Values="$(OptionalClientTarget);$(OptionalServerTarget)">
 			<Node Name="Compile $(Target) Win64" Requires="Update Version Files" Produces="#$(Target) Win64">
 				<ForEach Name="Configuration" Values="$(GameConfigurations)">
-					<Compile Target="$(Target)" Platform="Win64" Configuration="$(Configuration)" Tag="#$(Target) Win64" Arguments="-precompile -allmodules -nolink $(TargetDebugInfoArg) $(ExtraCompileArgs) $(Win64NoPCHArg)"/>
-					<Compile Target="$(Target)" Platform="Win64" Configuration="$(Configuration)" Tag="#$(Target) Win64" Arguments="-precompile $(TargetDebugInfoArg) $(ExtraCompileArgs) $(Win64NoPCHArg)" Clean="false"/>
+					<Compile Target="$(Target)" Platform="Win64" Configuration="$(Configuration)" Tag="#$(Target) Win64" Arguments="-precompile -architectures=$(Win64Architectures) -allmodules -nolink $(TargetDebugInfoArg) $(ExtraCompileArgs) $(Win64NoPCHArg)"/>
+					<Compile Target="$(Target)" Platform="Win64" Configuration="$(Configuration)" Tag="#$(Target) Win64" Arguments="-precompile -architectures=$(Win64Architectures) $(TargetDebugInfoArg) $(ExtraCompileArgs) $(Win64NoPCHArg)" Clean="false"/>
 					<Property Name="OptionalTargetTags" Value="$(OptionalTargetTags);#$(Target) Win64"/>
 				</ForEach>
 			</Node>
@@ -341,8 +341,8 @@
 				<Tag Files="$(Tag)" With="#UnrealGame Win64"/>
 			</ForEach>
 			<ForEach Name="Configuration" Values="$(GameConfigurations)">
-				<Compile Target="UnrealGame" Platform="Win64" Configuration="$(Configuration)" Tag="#UnrealGame Win64" AllowParallelExecutor="$(bAllowParallelExecutor)" Arguments="-precompile -allmodules -nolink $(TargetDebugInfoArg) $(ExtraCompileArgs) $(Win64NoPCHArg)"/>
-				<Compile Target="UnrealGame" Platform="Win64" Configuration="$(Configuration)" Tag="#UnrealGame Win64" AllowParallelExecutor="$(bAllowParallelExecutor)" Arguments="-precompile $(TargetDebugInfoArg) $(ExtraCompileArgs) $(Win64NoPCHArg)" Clean="false"/>
+				<Compile Target="UnrealGame" Platform="Win64" Configuration="$(Configuration)" Tag="#UnrealGame Win64" AllowParallelExecutor="$(bAllowParallelExecutor)" Arguments="-precompile -architectures=$(Win64Architectures) -allmodules -nolink $(TargetDebugInfoArg) $(ExtraCompileArgs) $(Win64NoPCHArg)"/>
+				<Compile Target="UnrealGame" Platform="Win64" Configuration="$(Configuration)" Tag="#UnrealGame Win64" AllowParallelExecutor="$(bAllowParallelExecutor)" Arguments="-precompile -architectures=$(Win64Architectures) $(TargetDebugInfoArg) $(ExtraCompileArgs) $(Win64NoPCHArg)" Clean="false"/>
 			</ForEach>
 
 			<Do If="$(EmbedSrcSrvInfo)">
