[EditorPerf]

- Remove -nozenloader option from EditorPerf tests as the comparison is no longer useful and the expense of running both loader configurations is non-trivial
- Fixes the path used when running with RaceDetector for finding log reports when running in CI vs locally

#rnx

diff --git a/Engine/Build/EditorPerf.xml b/Engine/Build/EditorPerf.xml
--- a/Engine/Build/EditorPerf.xml
+++ b/Engine/Build/EditorPerf.xml
@@ -33,6 +33,7 @@
 	<Option Name="NoDebugInfo" Restrict="true|false" DefaultValue="false" Description="Compile without including debug info"/>
 	<Option Name="WithTSAN" Restrict="true|false" DefaultValue="false" Description="Compile using TSAN"/>
 	<Option Name="WithRaceDetector" Restrict="true|false" DefaultValue="false" Description="Compile using Race Detector"/>
+	<Property Name="WithTSANorRaceDetector" Value="$(WithTSAN) OR $(WithRaceDetector)"/>
 	<Property Name="ExtraPluginsCommaDelims" Value="EditorTests"/>
 	<Property Name="ExtraPluginsPlusDelims" Value=""/> <!-- populated programmatically -->
 	<ForEach Name="Plugin" Values="$(ExtraPluginsCommaDelims)">
@@ -59,10 +60,10 @@
 	<Property Name="AutomationOutputDir" Value="$(AutomationLocalOutputDir)"/>
 	<Property Name="AutomationOutputDir" Value="$(RootDir)/Engine/Programs/AutomationTool/Saved/Logs" If="$(IsBuildMachine)"/>
 	<!-- We need to share logs between nodes for TSAN runs so store the logs to a share rather than Horde's Log storage -->
-	<Property Name="AutomationNetworkOutputDir" Value="$(AutomationNetworkRootDir)\$(BuildName)" If="$(WithTSAN)"/>
-	<Property Name="AutomationNetworkOutputDir" Value="$(AutomationNetworkRootDir)/$(BuildName)" If="$(WithTSAN) and '$(HostPlatform)' != 'Win64'"/>
-	<Property Name="AutomationOutputDir" Value="$(AutomationLocalOutputDir)" If="$(WithTSAN)"/>
-	<Property Name="AutomationOutputDir" Value="$(AutomationNetworkOutputDir)" If="$(WithTSAN) and $(IsBuildMachine)"/>
+	<Property Name="AutomationNetworkOutputDir" Value="$(AutomationNetworkRootDir)\$(BuildName)" If="$(WithTSANorRaceDetector)"/>
+	<Property Name="AutomationNetworkOutputDir" Value="$(AutomationNetworkRootDir)/$(BuildName)" If="$(WithTSANorRaceDetector) and '$(HostPlatform)' != 'Win64'"/>
+	<Property Name="AutomationOutputDir" Value="$(AutomationLocalOutputDir)" If="$(WithTSANorRaceDetector)"/>
+	<Property Name="AutomationOutputDir" Value="$(AutomationNetworkOutputDir)" If="$(WithTSANorRaceDetector) and $(IsBuildMachine)"/>
 
 	<!-- Options for warm+cold vs cold cache only -->
 	<Option Name="IncludeWarmPhase" DefaultValue="true" Description="Whether to include a warm cache phase. Default is true and will run cold and warm tests. If false, only a cold test will run."/>
@@ -74,17 +75,14 @@
 
 	<!-- Arguments for phases, keep in mind that these are passed to both the DDCCleanupCommandlet and RunUnreal -->
 	<Option Name="EnableBasePhaseTests" DefaultValue="true" Description="Enable the base phase perf tests."/>
-	<Property Name="EnableBasePhaseTests" Value="false" If="$(WithTSAN) OR $(WithRaceDetector)"/>
+	<Property Name="EnableBasePhaseTests" Value="false" If="$(WithTSANorRaceDetector)"/>
 	<Option Name="ColdArgs" DefaultValue="-ddc=Cold -NoAssetRegistryCacheRead" Description="Arguments to unreal for the cold phase"/>
 	<Option Name="WarmArgs" DefaultValue="-ddc=Warm" Description="Arguments to unreal for the warm phase"/>
 
-	<!-- Iterations to run and average together -->
-	<Option Name="Iterations" DefaultValue="1" Description="The list of iteration numbers to run. For example '1;2;3' for three iterations"/>
-
 	<!-- Option for number of PIE iterations to run (per editor launch) -->
 	<Property Name="DefaultNumIterationsPie" Value="2" />
 	<Option Name="NumIterationsPIE" DefaultValue="$(DefaultNumIterationsPIE)" Description="Specify the number of times to start/stop PIE sessions in a single editor run."/>
-	<Property Name="NumIterationsPIE" Value="1" If="$(WithTSAN) OR $(WithRaceDetector)"/>
+	<Property Name="NumIterationsPIE" Value="1" If="$(WithTSANorRaceDetector)"/>
 	<Do If="'$(NumIterationsPIE)' &lt;= '0'">
 		<Warning Message="NumIterationPIE must be greater than 0. Setting to default of $(DefaultNumIterationsPIE)."/>
 		<Property Name="NumIterationsPIE" Value="$(DefaultNumIterationsPIE)" />
@@ -157,13 +155,11 @@
 	<Option Name="AdditionalRuntimeArgs" DefaultValue="" Description="Additional Arguments to pass to RunTest"/>
 
 	<!-- Options for extra Zen Loader tests -->
-	<Option Name="EnableZenLoaderTest" DefaultValue="true" Description="Enable additional test per phase that uses the Zen Loader."/>
 	<Option Name="EnableZenLoaderAsyncThreadTest" DefaultValue="true" Description="Enable additional test per phase that uses the Zen Loader with an async loading thread."/>
 	<Option Name="DisableEditorDomainForZenLoaderTests" DefaultValue="true" Description="Explicitly disables editordomain (-noeditordomain) for zen loader tests"/>
 
 	<!-- Zen Loader tests specifc arguments -->
 	<Property Name="ZenLoaderTestArgs" Value="-zenloader"/>
-	<Property Name="NoZenLoaderTestArgs" Value="-nozenloader"/>
 	<Property Name="ZenLoaderAsyncTestArgs" Value="$(ZenLoaderTestArgs) -asyncloadingthread -dpcvars=s.AllowMultithreadedLoading=1" />
 	<Property Name="RaceDetectorTestArgs" Value="-racedetector" />
 	
@@ -180,7 +176,6 @@
 	<Property Name="ZenStoreDDCTestArgs" Value="$(ZenStoreCommonArgs)"/>
 
 	<!-- Options for extra LazyLoad tests -->
-	<Option Name="EnableLazyLoadNoZenLoaderTest" DefaultValue="false" Description="Enable additional test on warm phase that uses LazyLoad with Zen Loader disabled."/>
 	<Option Name="EnableLazyLoadZenLoaderTest" DefaultValue="false" Description="Enable additional test on warm phase that uses LazyLoad with Zen Loader enabled."/>
 	
 	<!-- LazyLoad tests specific arguments -->
@@ -235,11 +230,11 @@
 		
 		<!-- Don't send telemetry for TSAN jobs as the time won't be comparable -->
 		<Property Name="EditorTelemetryArgs" Value=""/>
-		<Property Name="EditorTelemetryArgs" Value="-ST_SendTelemetry=false" If="$(WithTSAN)"/>
+		<Property Name="EditorTelemetryArgs" Value="-ST_SendTelemetry=false" If="$(WithTSANorRaceDetector)"/>
 
 		<!-- Don't itemize timers for TSAN jobs as the time won't be comparable -->
 		<Property Name="UploadToPRSEffective" Value="$(UploadToPRS)"/>
-		<Property Name="UploadToPRSEffective" Value="false" If="$(WithTSAN)"/>
+		<Property Name="UploadToPRSEffective" Value="false" If="$(WithTSANorRaceDetector)"/>
 		
 		<Property Name="TestFilterArgs" Value=""/>
 		<Property Name="CombinedEditorArgs" Value=""/>
@@ -432,42 +427,24 @@
 					<!-- Non-TSAN Nodes -->
 					<Do If="!$(WithTSAN) AND !$(WithRaceDetector)">
 						<!-- 
-						Note, each of the nodes has a separate ForEach for Phases (where appropriate) because we want each node type (zen, nozen, async etc...) to be paired: Cold then Warm phase nodes.
+						Note, each of the nodes has a separate ForEach for Phases (where appropriate) because we want each node type (zen, async etc...) to be paired: Cold then Warm phase nodes.
 						If we did not order nodes like this, since the nodes run on the same agent it's possible for a cold job to fail and prevent a warm job _from an unrelated node_ to not have 
 						a properly populated DDC, defeating the purpose of the warm run. Thus all node type must run cold then warm, one after the other.
 						-->
 
 						<Property Name="RunNodeCompilationRequires" Value="$(CompileNodeName)"/>
 						<Property Name="RunNodeCompilationRequires" Value="$($($(Project)__UseBinariesFrom)__CompileNodeName)" If="'$($(Project)__UseBinariesFrom)'!=''"/>
-						<Do If="$(EnableBasePhaseTests)">
-							<ForEach Name="Phase" Values="$(Phases)">
-								<ForEach Name="Iteration" Values="$(Iterations)">
-									<Property Name="IterationSuffix" Value=""/>
-									<Property Name="IterationSuffix" Value="$(Iteration)" If="'$(Iteration)' != '1'"/>
-									<Property Name="RunNodeName" Value="Run $(OutputName) $(Phase)$(IterationSuffix)"/>
-									<Property Name="RunNodeRequires" Value="$(RunNodeCompilationRequires)"/>
-									<Property Name="RunNodeRequires" Value="$(RunNodeCompilationRequires); Run $(OutputName) Cold$(IterationSuffix)" If="'$(Phase)' == 'Warm'"/>
-									<Node Name="$(RunNodeName)" Requires="$(RunNodeRequires)" IgnoreModified="$(ProjectBinariesPath)/...">
-										<Expand Name="RunTest" TestNameSuffix="NoZenLoader" Platform="$(Platform)" ProjectName="$(Project)" ProjectTarget="$(ProjectTarget)" NodeName="$(RunNodeName)" Phase="$(Phase)" Async="false" UploadToPRS="false" Iteration="$(Iteration)" Args="$($(Phase)Args)" EditorArgs="$(ProjectCombinedArgs) $(NoZenLoaderTestArgs)"/>
-									</Node>
-									<Property Name="RunNodeList" Value="$(RunNodeList);$(RunNodeName)"/>
-								</ForEach>
-							</ForEach>
-						</Do>
 
-						<!--Zen Loader TestRun if enabled-->
-						<Do If="$(EnableZenLoaderTest)">
-							<!-- Setup Zen Loader tests for this phase -->
-							<ForEach Name="Phase" Values="$(Phases)">
-								<Property Name="RunNodeZenName" Value="Run $(OutputName) $(Phase) ZenLoader"/>
-								<Property Name="RunNodeZenRequires" Value="$(RunNodeCompilationRequires)"/>
-								<Property Name="RunNodeZenRequires" Value="$(RunNodeCompilationRequires); Run $(OutputName) Cold ZenLoader" If="'$(Phase)' == 'Warm'"/>
-								<Node Name="$(RunNodeZenName)" Requires="$(RunNodeZenRequires)" IgnoreModified="$(ProjectBinariesPath)/...">
-									<Expand Name="RunTest" TestNameSuffix="ZenLoader" Platform="$(Platform)" ProjectName="$(Project)" ProjectTarget="$(ProjectTarget)" NodeName="$(RunNodeZenName)" Phase="$(Phase)" Async="false" UploadToPRS="$(UploadTraceTimersToPRS)" Iteration="1" Args="$($(Phase)Args)" EditorArgs="$(ProjectCombinedArgs) $(ZenLoaderTestArgs)"/>
-								</Node>
-								<Property Name="RunNodeList" Value="$(RunNodeList);$(RunNodeZenName)" If="$(EnableZenLoaderTest)"/>
-							</ForEach>
-						</Do>
+						<!--Zen Loader TestRun -->
+						<ForEach Name="Phase" Values="$(Phases)">
+							<Property Name="RunNodeZenName" Value="Run $(OutputName) $(Phase) ZenLoader"/>
+							<Property Name="RunNodeZenRequires" Value="$(RunNodeCompilationRequires)"/>
+							<Property Name="RunNodeZenRequires" Value="$(RunNodeCompilationRequires); Run $(OutputName) Cold ZenLoader" If="'$(Phase)' == 'Warm'"/>
+							<Node Name="$(RunNodeZenName)" Requires="$(RunNodeZenRequires)" IgnoreModified="$(ProjectBinariesPath)/...">
+								<Expand Name="RunTest" TestNameSuffix="ZenLoader" Platform="$(Platform)" ProjectName="$(Project)" ProjectTarget="$(ProjectTarget)" NodeName="$(RunNodeZenName)" Phase="$(Phase)" Async="false" UploadToPRS="$(UploadTraceTimersToPRS)" Iteration="1" Args="$($(Phase)Args)" EditorArgs="$(ProjectCombinedArgs) $(ZenLoaderTestArgs)"/>
+							</Node>
+							<Property Name="RunNodeList" Value="$(RunNodeList);$(RunNodeZenName)"/>
+						</ForEach>
 
 						<!--Zen Loader Async TestRun if enabled-->
 						<Do If="$(EnableZenLoaderAsyncThreadTest)">
@@ -497,16 +474,7 @@
 
 						<!--Warm only test runs e.g LazyLoad-->
 						<ForEach Name="Phase" Values="Warm">
-							<!--Legacy LazyLoad TestRun if enabled-->
-							<Property Name="RunNodeLazyLoadName" Value="Run $(OutputName) $(Phase) LazyLoad NoZenLoader"/>
-							<Property Name="RunNodeLazyLoadRequires" Value="$(RunNodeCompilationRequires)"/>
-							<Property Name="RunNodeLazyLoadRequires" Value="$(RunNodeCompilationRequires); Run $(OutputName) Cold LazyLoad NoZenLoader" If="'$(Phase)' == 'Warm'"/>
-							<Node Name="$(RunNodeLazyLoadName)" Requires="$(RunNodeLazyLoadRequires)" IgnoreModified="$(ProjectBinariesPath)/..." If="$(EnableLazyLoadNoZenLoaderTest)">
-								<Expand Name="RunTest" TestNameSuffix="NoZenLoader" Platform="$(Platform)" ProjectName="$(Project)" ProjectTarget="$(ProjectTarget)" NodeName="$(RunNodeLazyLoadName)" Phase="$(Phase)" Async="false" UploadToPRS="false" Iteration="1" Args="$($(Phase)Args)" EditorArgs="$(ProjectCombinedArgs) $(LazyLoadTestArgs) $(NoZenLoaderTestArgs)"/>
-							</Node>
-							<Property Name="RunNodeList" Value="$(RunNodeList);$(RunNodeLazyLoadName)" If="$(EnableLazyLoadNoZenLoaderTest)"/>
-
-							<!--Zen LazyLoad TestRun if enabled-->
+							<!--Zen LazyLoad TestRun-->
 							<Property Name="RunNodeLazyLoadZenName" Value="Run $(OutputName) $(Phase) LazyLoad ZenLoader"/>
 							<Property Name="RunNodeLazyLoadZenRequires" Value="$(RunNodeCompilationRequires)"/>
 							<Property Name="RunNodeLazyLoadZenRequires" Value="$(RunNodeCompilationRequires); Run $(OutputName) Cold LazyLoad ZenLoader" If="'$(Phase)' == 'Warm'"/>
