Ensure Editor manifest is available in base engine package
#rnx

diff --git a/Engine/Build/InstalledEngineFilters.xml b/Engine/Build/InstalledEngineFilters.xml
--- a/Engine/Build/InstalledEngineFilters.xml
+++ b/Engine/Build/InstalledEngineFilters.xml
@@ -162,6 +162,7 @@
 
 		<!-- Only standalone file required for Engine locs -->
 		Engine/Content/Localization/Engine/Engine.manifest
+		Engine/Content/Localization/Editor/Editor.manifest
 </Property>
 	
 	<!-- Files which can exist under any engine or engine-platform directory -->
