InstalledBuild, don't build iPhonePackager CsProj's as this conflicts with checked binary prebuilts. Instead, reference these as was done in the past with the .NET Framework version of iPhonePackager.

#Synced-CL 48743327

diff --git a/Engine/Build/InstalledEngineBuild.xml b/Engine/Build/InstalledEngineBuild.xml
--- a/Engine/Build/InstalledEngineBuild.xml
+++ b/Engine/Build/InstalledEngineBuild.xml
@@ -594,18 +594,17 @@
 
 			<!-- Compile all the tools -->
 			<CsCompile Project="#RedistUATProjects" Configuration="Development" Platform="AnyCPU" Tag="#Build Tools CS Binaries;#Sign Binaries" TagReferences="#Build Tools CS Binaries"/>
-			<Do If="'$(HostPlatform)' == 'Win64'">
 			<!-- TODO: needs to be updated to dotnet8 or deleted
+			<Do If="'$(HostPlatform)' == 'Win64'">
 				<CsCompile Project="$(CsToolsDir)/Engine/Source/Programs/UnrealSwarm/SwarmCoordinator/SwarmCoordinator.csproj" Configuration="Development" Platform="AnyCPU" UseSystemCompiler="true" Tag="#Build Tools CS Binaries;#Sign Binaries" TagReferences="#Build Tools CS Binaries"/>
 				<CsCompile Project="$(CsToolsDir)/Engine/Source/Programs/UnrealSwarm/Agent/Agent.csproj" Configuration="Development" Platform="AnyCPU" UseSystemCompiler="true" Tag="#Build Tools CS Binaries;#Sign Binaries" TagReferences="#Build Tools CS Binaries"/>
 				<CsCompile Project="$(CsToolsDir)/Engine/Source/Editor/SwarmInterface/DotNET/SwarmInterface.csproj" Configuration="Development" Platform="AnyCPU" UseSystemCompiler="true" Tag="#Build Tools CS Binaries;#Sign Binaries" TagReferences="#Build Tools CS Binaries"/>
 				<CsCompile Project="$(CsToolsDir)/Engine/Source/Programs/UnrealControls/UnrealControls.csproj" Configuration="Development" Platform="AnyCPU" UseSystemCompiler="true" Tag="#Build Tools CS Binaries;#Sign Binaries" TagReferences="#Build Tools CS Binaries"/>
-				<CsCompile Project="$(CsToolsDir)/Engine/Source/Programs/NetworkProfiler/NetworkProfiler/NetworkProfiler.csproj" Configuration="Development" Platform="AnyCPU" UseSystemCompiler="true" Tag="#Build Tools CS Binaries" TagReferences="#Build Tools CS Binaries;#Sign Binaries"/>
-			-->
 				<CsCompile Project="$(CsToolsDir)/Engine/Source/Programs/IOS/iPhonePackager/iPhonePackager.csproj" Configuration="Development" Platform="AnyCPU" UseSystemCompiler="true" Arguments="/verbosity:minimal /target:Rebuild" Tag="#Build Tools CS Binaries;#Sign Binaries" TagReferences="#Build Tools CS Binaries"/>
-				<Tag Files="$(CsToolsDir)/Engine/Binaries/DotNET/IOS/..." With="#Build Tools CS Binaries" />
+				<CsCompile Project="$(CsToolsDir)/Engine/Source/Programs/NetworkProfiler/NetworkProfiler/NetworkProfiler.csproj" Configuration="Development" Platform="AnyCPU" UseSystemCompiler="true" Tag="#Build Tools CS Binaries" TagReferences="#Build Tools CS Binaries;#Sign Binaries"/>
 				<Expand Name="Platform_Build_Tools_CS_Binaries"/>
 			</Do>
+			-->
 
 			<!-- Save and tag BuildConfiguration.Schema.xsd as this can/will be overriden by other processes piggy backing on cstools to get it to the right spot -->
 			<Copy Files="$(RootDir)/Engine/Saved/UnrealBuildTool/BuildConfiguration.Schema.xsd" From="$(RootDir)" To="$(CsToolsDir)"/>
@@ -685,10 +684,16 @@
 			<!-- Tag OidcToken which is already built for all platforms and just needs to be copied  -->
 			<Tag Files="$(CsToolsDir)/Engine/Binaries/DotNET/OidcToken/..." With="#Build Tools CS Binaries Mac" />
 
+			<!-- IPP could be built here but is checked in prebuilt. -->
 			<!-- CsCompile Project="$(CsToolsDir)/Engine/Source/Programs/UnrealControls/UnrealControls.csproj" Configuration="Development" Platform="AnyCPU" Tag="#Build Tools CS Binaries Mac;#Sign Binaries Mac" TagReferences="#Build Tools CS Binaries Mac"/-->
-			<CsCompile Project="$(CsToolsDir)/Engine/Source/Programs/IOS/iPhonePackager/iPhonePackagerCLI.csproj" Configuration="Development" Platform="AnyCPU" Arguments="/verbosity:minimal /target:Rebuild" Tag="#Build Tools CS Binaries Mac;#Sign Binaries Mac" TagReferences="#Build Tools CS Binaries Mac"/>
-			<Tag Files="$(CsToolsDir)/Engine/Binaries/DotNET/IOS/..." With="#Build Tools CS Binaries Mac" />
-
+			<!-- CsCompile Project="$(CsToolsDir)/Engine/Source/Programs/IOS/iPhonePackager/iPhonePackager.csproj" Configuration="Development" Platform="AnyCPU" Arguments="/verbosity:minimal /target:Rebuild" Tag="#Build Tools CS Binaries Mac;#Sign Binaries Mac" TagReferences="#Build Tools CS Binaries Mac"/-->
+		
+			<!-- So we will ship the binary out of Perforce for now. -->
+			<Tag Files="Engine/Binaries/DotNET/IOS/..." With="#CopyIPPBinaries"/>
+			<Copy Files="#CopyIPPBinaries" From="$(RootDir)" To="$(CsToolsDir)"/>
+			<Tag Files="$(CsToolsDir)/Engine/Binaries/DotNET/IOS/IPhonePackagerCLI.exe" With="#Build Tools CS Binaries Mac;#Sign Binaries Mac"/>
+			<Tag Files="$(CsToolsDir)/Engine/Binaries/DotNET/IOS/IPhonePackagerCLI.runtimeconfig.json" With="#Build Tools CS Binaries Mac;#Sign Binaries Mac"/>
+			
 			<!-- Sign the binaries -->
 			<Sign Files="#Sign Binaries Mac" If="'$(SignExecutables)' == true"/>
 		</Node>
