iPhonePackager for Mac is just in dll form now so just copy.

#Synced-CL 48743327

diff --git a/Engine/Build/InstalledEngineBuild.xml b/Engine/Build/InstalledEngineBuild.xml
--- a/Engine/Build/InstalledEngineBuild.xml
+++ b/Engine/Build/InstalledEngineBuild.xml
@@ -691,8 +691,7 @@
 			<!-- So we will ship the binary out of Perforce for now. -->
 			<Tag Files="Engine/Binaries/DotNET/IOS/..." With="#CopyIPPBinaries"/>
 			<Copy Files="#CopyIPPBinaries" From="$(RootDir)" To="$(CsToolsDir)"/>
-			<Tag Files="$(CsToolsDir)/Engine/Binaries/DotNET/IOS/IPhonePackagerCLI.exe" With="#Build Tools CS Binaries Mac;#Sign Binaries Mac"/>
-			<Tag Files="$(CsToolsDir)/Engine/Binaries/DotNET/IOS/IPhonePackagerCLI.runtimeconfig.json" With="#Build Tools CS Binaries Mac;#Sign Binaries Mac"/>
+			<Tag Files="$(CsToolsDir)/Engine/Binaries/DotNET/IOS/..." With="#Build Tools CS Binaries Mac"/>
 			
 			<!-- Sign the binaries -->
 			<Sign Files="#Sign Binaries Mac" If="'$(SignExecutables)' == true"/>
