[EditorPerf]

Make ZenLoader arguments explicit about using/not using the asyncloading thread and wide loader so if global defaults change the jobs remain unaffected

#rnx
#Synced-CL 49919180

diff --git a/Engine/Build/EditorPerf.xml b/Engine/Build/EditorPerf.xml
--- a/Engine/Build/EditorPerf.xml
+++ b/Engine/Build/EditorPerf.xml
@@ -159,7 +159,7 @@
 	<Option Name="DisableEditorDomainForZenLoaderTests" DefaultValue="true" Description="Explicitly disables editordomain (-noeditordomain) for zen loader tests"/>
 
 	<!-- Zen Loader tests specifc arguments -->
-	<Property Name="ZenLoaderTestArgs" Value="-zenloader"/>
+	<Property Name="ZenLoaderTestArgs" Value="-zenloader -NoAsyncLoadingThread -dpcvars=s.AllowMultithreadedLoading=0"/>
 	<Property Name="ZenLoaderAsyncTestArgs" Value="$(ZenLoaderTestArgs) -asyncloadingthread -dpcvars=s.AllowMultithreadedLoading=1" />
 	<Property Name="RaceDetectorTestArgs" Value="-racedetector" />
 	
