[EditorPerf][CI]

Prevent contradictory cmdline arguments from being combined when using async loading as not all argument handlers take last argument as the highest priority

#rnx

#Synced-CL 50295700

diff --git a/Engine/Build/EditorPerf.xml b/Engine/Build/EditorPerf.xml
--- a/Engine/Build/EditorPerf.xml
+++ b/Engine/Build/EditorPerf.xml
@@ -160,7 +160,7 @@
 
 	<!-- Zen Loader tests specifc arguments -->
 	<Property Name="ZenLoaderTestArgs" Value="-zenloader -NoAsyncLoadingThread -dpcvars=s.AllowMultithreadedLoading=0"/>
-	<Property Name="ZenLoaderAsyncTestArgs" Value="$(ZenLoaderTestArgs) -asyncloadingthread -dpcvars=s.AllowMultithreadedLoading=1" />
+	<Property Name="ZenLoaderAsyncTestArgs" Value="-zenloader  -asyncloadingthread -dpcvars=s.AllowMultithreadedLoading=1" />
 	<Property Name="RaceDetectorTestArgs" Value="-racedetector" />
 	
 	<Property Name="ZenLoaderTestArgs" Value="$(ZenLoaderTestArgs) -noeditordomain" If="$(DisableEditorDomainForZenLoaderTests)" />
