Adding missing embree dependencies

#rnx

diff --git a/Engine/Build/InstalledEngineFilters.xml b/Engine/Build/InstalledEngineFilters.xml
--- a/Engine/Build/InstalledEngineFilters.xml
+++ b/Engine/Build/InstalledEngineFilters.xml
@@ -375,6 +375,9 @@
 
 		<!-- DiaSdk -->
 		Engine/Source/ThirdParty/DiaSdk/...
+
+		<!-- Embree -->
+		Engine/Source/ThirdParty/Intel/Embree/Deploy/embree-4.3.3/Unix/x86_64-unknown-linux-gnu/lib/libembree4.so
 	</Property>
 
 	<Property Name="CopyEditorExceptionsWin64" Value="$(CopyEditorExceptions)">
@@ -443,6 +446,9 @@
 
 		<!-- debug tools -->
 		Engine/Source/ThirdParty/Breakpad/src/tools/mac/binaries/dump_syms
+
+		<!-- Embree -->
+		Engine/Source/ThirdParty/Intel/Embree/Deploy/embree-4.3.3/Unix/x86_64-unknown-linux-gnu/lib/libembree4.so
 	</Property>
 
 	<Property Name="CopyEditorExceptionsMac" Value="$(CopyEditorExceptions)">
@@ -528,6 +534,9 @@
 
 		<!-- TechSoft -->
 		Engine/Binaries/ThirdParty/TechSoft/Linux/...
+
+		<!-- Embree -->
+		Engine/Source/ThirdParty/Intel/Embree/Deploy/embree-4.3.3/Unix/x86_64-unknown-linux-gnu/lib/libembree4.so
 	</Property>
 	<Property Name="CopyEditorExceptionsLinux" Value="$(CopyEditorExceptions)">
 		<!-- Exclude Windows/Mac binaries on Linux here -->
