Add full Unix embree binaries to win/mac/linux engine install + ensure they're only downloadable via linux installer option in launcher
#rnx

diff --git a/Engine/Build/InstalledEngineFilters.xml b/Engine/Build/InstalledEngineFilters.xml
--- a/Engine/Build/InstalledEngineFilters.xml
+++ b/Engine/Build/InstalledEngineFilters.xml
@@ -377,7 +377,7 @@
 		Engine/Source/ThirdParty/DiaSdk/...
 
 		<!-- Embree -->
-		Engine/Source/ThirdParty/Intel/Embree/Deploy/embree-4.3.3/Unix/x86_64-unknown-linux-gnu/lib/libembree4.so
+		Engine/Source/ThirdParty/Intel/Embree/Deploy/embree-4.3.3/Unix/x86_64-unknown-linux-gnu/...
 	</Property>
 
 	<Property Name="CopyEditorExceptionsWin64" Value="$(CopyEditorExceptions)">
@@ -448,7 +448,7 @@
 		Engine/Source/ThirdParty/Breakpad/src/tools/mac/binaries/dump_syms
 
 		<!-- Embree -->
-		Engine/Source/ThirdParty/Intel/Embree/Deploy/embree-4.3.3/Unix/x86_64-unknown-linux-gnu/lib/libembree4.so
+		Engine/Source/ThirdParty/Intel/Embree/Deploy/embree-4.3.3/Unix/x86_64-unknown-linux-gnu/...
 	</Property>
 
 	<Property Name="CopyEditorExceptionsMac" Value="$(CopyEditorExceptions)">
@@ -536,7 +536,7 @@
 		Engine/Binaries/ThirdParty/TechSoft/Linux/...
 
 		<!-- Embree -->
-		Engine/Source/ThirdParty/Intel/Embree/Deploy/embree-4.3.3/Unix/x86_64-unknown-linux-gnu/lib/libembree4.so
+		Engine/Source/ThirdParty/Intel/Embree/Deploy/embree-4.3.3/Unix/x86_64-unknown-linux-gnu/...
 	</Property>
 	<Property Name="CopyEditorExceptionsLinux" Value="$(CopyEditorExceptions)">
 		<!-- Exclude Windows/Mac binaries on Linux here -->
