Low level test report key adjustment

#rnx
#Synced-CL 0

diff --git a/Engine/Build/LowLevelTests.xml b/Engine/Build/LowLevelTests.xml
--- a/Engine/Build/LowLevelTests.xml
+++ b/Engine/Build/LowLevelTests.xml
@@ -248,7 +248,7 @@
 						<Property Name="TestArgsExt" Value="$(TestArgsExt) -adbkeys=$(AdbKeysPath) -captureoutput"/>
 					</Do>
 
-					<Property Name="TestDataKey" Value="$(ShortName) Tests $(Platform)$(PreflightSuffix)"/>
+					<Property Name="TestDataKey" Value="$(ShortName) Tests$(PreflightSuffix)"/>
 
 					<Command Name="RunLowLevelTests"
 							Arguments="-test=LowLevelTests -testapp=$(TargetName)$(TestModeSuffix) -platform=$(Platform) -configuration=$(Configuration) $(TestArgsExt) -HordeTestDataKey=&quot;$(TestDataKey)&quot; -VeryVerbose" If="!$(DryRun) and '$(RunUnsupported)' == ''"/>
