add option to build with 2022 toolchain for win64 precompile targets

#rnx

#Synced-CL 51916670

diff --git a/Engine/Build/InstalledEngineBuild.xml b/Engine/Build/InstalledEngineBuild.xml
--- a/Engine/Build/InstalledEngineBuild.xml
+++ b/Engine/Build/InstalledEngineBuild.xml
@@ -51,6 +51,10 @@
 	<Property Name="Win64NoPCHArg" Value="" />
 	<Property Name="Win64NoPCHArg" Value="-NoPCH" If="$(WithWin64NoPCH)"/>
 
+	<Option Name="Use2022Toolchain" Restrict="true|false" DefaultValue="false" Description="Whether to build precompile targets with the vs2022 toolchain"/>
+	<Property Name="VS2022Arg" Value=""/>
+	<Property Name="VS2022Arg" Value="-2022" If="$(Use2022Toolchain)"/>
+
 	<!-- Host platforms for the installed editor -->
 	<Option Name="WithInstalledWin64" Restrict="true|false" DefaultValue="true" Description="Whether to include the Win64 editor as part of a binary release"/>
 	<Option Name="WithInstalledMac" Restrict="true|false" DefaultValue="true" Description="Whether to include the Mac editor as part of a binary release"/>
@@ -215,11 +219,11 @@
 			<!-- We invoke QueryTargets to generate LiveLinkHub's TargetInfo.json file which would otherwise be missing since we're invoking UBT without a -project arg. -->
 			<QueryTargets ProjectFile="Engine/Source/Programs/LiveLinkHub/LiveLinkHub.uproject" Tag="#UnrealEditor Win64"/>
 			
-			<Compile Target="UnrealEditor" Platform="Win64" Configuration="DebugGame" Tag="#UnrealEditor Win64" Arguments="-precompile -allmodules $(ExtraCompileArgs) $(Win64NoPCHArg)"/>
-			<Compile Target="UnrealEditor" Platform="Win64" Configuration="Development" Tag="#UnrealEditor Win64" Arguments="-precompile -allmodules $(ExtraCompileArgs) $(Win64NoPCHArg)"/>
+			<Compile Target="UnrealEditor" Platform="Win64" Configuration="DebugGame" Tag="#UnrealEditor Win64" Arguments="-precompile -allmodules $(ExtraCompileArgs) $(Win64NoPCHArg) $(VS2022Arg)"/>
+			<Compile Target="UnrealEditor" Platform="Win64" Configuration="Development" Tag="#UnrealEditor Win64" Arguments="-precompile -allmodules $(ExtraCompileArgs) $(Win64NoPCHArg) $(VS2022Arg)"/>
 			
 			<!-- Note: We compile LiveLinkHub here because we share a modular build environment with the editor. This drastically reduces its build time and prevents having duplicate files in the MakeInstalledBuild step. -->
-			<Compile Target="LiveLinkHub" Configuration="Development" Platform="Win64" Tag="#UnrealEditor Win64" Arguments="-precompile $(ExtraCompileArgs) $(Win64NoPCHArg)" />
+			<Compile Target="LiveLinkHub" Configuration="Development" Platform="Win64" Tag="#UnrealEditor Win64" Arguments="-precompile $(ExtraCompileArgs) $(Win64NoPCHArg) $(VS2022Arg)" />
 			<Tag Files="Engine/Source/Programs/LiveLinkHub/Binaries/Win64/LiveLinkHub.target" With="#UnrealEditor Win64"/>
 
 			<Do If="$(EmbedSrcSrvInfo)">
@@ -330,8 +334,8 @@
 		<ForEach Name="Target" Values="$(OptionalClientTarget);$(OptionalServerTarget)">
 			<Node Name="Compile $(Target) Win64" Requires="Update Version Files" Produces="#$(Target) Win64">
 				<ForEach Name="Configuration" Values="$(GameConfigurations)">
-					<Compile Target="$(Target)" Platform="Win64" Configuration="$(Configuration)" Tag="#$(Target) Win64" Arguments="-precompile -architectures=$(Win64Architectures) -allmodules -nolink $(TargetDebugInfoArg) $(ExtraCompileArgs) $(Win64NoPCHArg)"/>
-					<Compile Target="$(Target)" Platform="Win64" Configuration="$(Configuration)" Tag="#$(Target) Win64" Arguments="-precompile -architectures=$(Win64Architectures) $(TargetDebugInfoArg) $(ExtraCompileArgs) $(Win64NoPCHArg)" Clean="false"/>
+					<Compile Target="$(Target)" Platform="Win64" Configuration="$(Configuration)" Tag="#$(Target) Win64" Arguments="-precompile -architectures=$(Win64Architectures) -allmodules -nolink $(TargetDebugInfoArg) $(ExtraCompileArgs) $(Win64NoPCHArg) $(VS2022Arg)"/>
+					<Compile Target="$(Target)" Platform="Win64" Configuration="$(Configuration)" Tag="#$(Target) Win64" Arguments="-precompile -architectures=$(Win64Architectures) $(TargetDebugInfoArg) $(ExtraCompileArgs) $(Win64NoPCHArg) $(VS2022Arg)" Clean="false"/>
 					<Property Name="OptionalTargetTags" Value="$(OptionalTargetTags);#$(Target) Win64"/>
 				</ForEach>
 			</Node>
@@ -342,8 +346,8 @@
 			</ForEach>
 			<ForEach Name="Configuration" Values="$(GameConfigurations)">
 				<ForEach Name="Architecture" Values="$(Win64Architectures)" Separator="+">
-					<Compile Target="UnrealGame" Platform="Win64" Configuration="$(Configuration)" Tag="#UnrealGame Win64" AllowParallelExecutor="$(bAllowParallelExecutor)" Arguments="-precompile -architectures=$(Architecture) -allmodules -nolink $(TargetDebugInfoArg) $(ExtraCompileArgs) $(Win64NoPCHArg)"/>
-					<Compile Target="UnrealGame" Platform="Win64" Configuration="$(Configuration)" Tag="#UnrealGame Win64" AllowParallelExecutor="$(bAllowParallelExecutor)" Arguments="-precompile -architectures=$(Architecture) $(TargetDebugInfoArg) $(ExtraCompileArgs) $(Win64NoPCHArg)" Clean="false"/>
+					<Compile Target="UnrealGame" Platform="Win64" Configuration="$(Configuration)" Tag="#UnrealGame Win64" AllowParallelExecutor="$(bAllowParallelExecutor)" Arguments="-precompile -architectures=$(Architecture) -allmodules -nolink $(TargetDebugInfoArg) $(ExtraCompileArgs) $(Win64NoPCHArg) $(VS2022Arg)"/>
+					<Compile Target="UnrealGame" Platform="Win64" Configuration="$(Configuration)" Tag="#UnrealGame Win64" AllowParallelExecutor="$(bAllowParallelExecutor)" Arguments="-precompile -architectures=$(Architecture) $(TargetDebugInfoArg) $(ExtraCompileArgs) $(Win64NoPCHArg) $(VS2022Arg)" Clean="false"/>
 				</ForEach>
 			</ForEach>
 
