Horde test key consistent for all RunLowLevelTests commands

#rnx
#Synced-CL 0

diff --git a/Engine/Build/LowLevelTests.xml b/Engine/Build/LowLevelTests.xml
--- a/Engine/Build/LowLevelTests.xml
+++ b/Engine/Build/LowLevelTests.xml
@@ -84,7 +84,6 @@
 
 	<!-- Check if this is a release branch -->
 	<Option Name="IsReleaseBranch" DefaultValue="false" Description="True if release branch, false otherwise"/>
-	<Option Name="PreflightChange" DefaultValue="" Description="Preflight CL number if preflight, empty otherwise"/>
 
 	<!-- Test Options -->
 	<Option Name="RunAllTests" DefaultValue="true" Description="Run All Available Tests (except Foundation)"/>
