Fixing installed engine build missing files when only linuxarm64 is enabled for windows

#rnx

diff --git a/Engine/Build/InstalledEngineBuild.xml b/Engine/Build/InstalledEngineBuild.xml
--- a/Engine/Build/InstalledEngineBuild.xml
+++ b/Engine/Build/InstalledEngineBuild.xml
@@ -1032,8 +1032,9 @@
 		<Property Name="InstalledRequirements" Value="$(InstalledRequirements);#UnrealGame Android;#UnrealGame Android Unstripped;#UnrealGame Android Stripped" If="'$(WithAndroid)' == true"/>
 		<Property Name="InstalledRequirements" Value="$(InstalledRequirements);#UnrealGame IOS;#UnrealGame IOS Unstripped;#UnrealGame IOS Stripped;#Update Version Files" If="'$(WithIOS)' == true"/>
 		<Property Name="InstalledRequirements" Value="$(InstalledRequirements);#UnrealGame TVOS;#UnrealGame TVOS Unstripped;#UnrealGame TVOS Stripped" If="'$(WithTVOS)' == true"/>
-		<Property Name="InstalledRequirements" Value="$(InstalledRequirements);#UnrealGame Linux;#UnrealGame Linux Unstripped;#UnrealGame Linux Stripped;#Build Tools Linux" If="'$(WithLinux)' == true"/>
+		<Property Name="InstalledRequirements" Value="$(InstalledRequirements);#UnrealGame Linux;#UnrealGame Linux Unstripped;#UnrealGame Linux Stripped" If="'$(WithLinux)' == true"/>
 		<Property Name="InstalledRequirements" Value="$(InstalledRequirements);#UnrealGame LinuxArm64;#UnrealGame LinuxArm64 Unstripped;#UnrealGame LinuxArm64 Stripped" If="'$(WithLinuxArm64)' == true"/>
+		<Property Name="InstalledRequirements" Value="$(InstalledRequirements);#Build Tools Linux" If="'$(WithLinux)' == true Or '$(WithLinuxArm64)' == true"/>
 		<Property Name="InstalledRequirements" Value="$(InstalledRequirements);#CompressedDDCWin64" If="$(WithDDC) And $(WithWin64)"/>
 		<Property Name="InstalledRequirements" Value="$(InstalledRequirements);#Compile Datasmith Plugins Win64;#DatasmithPluginBinaries Win64" If="$(WithWin64) And $(CompileDatasmithPlugins)"/>
 		<!-- Feature packs -->
@@ -1100,7 +1101,7 @@
 				<Property Name="ExtraSignFiles" Value="$(ExtraSignFiles);$(ExtraSignFilesLinux)"/>
 			</Do>
 			<Do If="'$(WithLinuxArm64)' == true">
-				<Tag Files="#UnrealGame LinuxArm64" Except="#UnrealGame LinuxArm64 Unstripped" With="#Installed Win64"/>
+				<Tag Files="#UnrealGame LinuxArm64;#Build Tools Linux" Except="#UnrealGame LinuxArm64 Unstripped" With="#Installed Win64"/>
 				<Tag Files="#UnrealGame LinuxArm64 Stripped" With="#Saved Output"/>
 				<Property Name="Platforms" Value="$(Platforms)LinuxArm64;"/>
 				<Property Name="CopyInstalledFilter" Value="$(CopyInstalledFilter);$(CopyLinuxFilter)"/>
