Ensure rsync 2.6.9 (used in RemoteMac build as a workaround to issues with the system rsync update in MacOs 26) is included in the Windows install build.

#Synced-CL 52308170

diff --git a/Engine/Build/InstalledEngineFilters.xml b/Engine/Build/InstalledEngineFilters.xml
--- a/Engine/Build/InstalledEngineFilters.xml
+++ b/Engine/Build/InstalledEngineFilters.xml
@@ -348,6 +348,7 @@
 
 		<!-- Used by Switchboard application, and by iOS/tvOS build automation -->
 		Engine/Extras/ThirdPartyNotUE/cwrsync/...
+		Engine/Extras/ThirdPartyNotUE/cwrsync/mac/...
 		Engine/Binaries/ThirdParty/PresentMon/Win64/PresentMon-1.8.0-x64.exe
 
 		<!-- nanosvg -->
